Added retry loops to SCW file IO

diff --git a/Engine/Source/Programs/ShaderCompileWorker/Private/ShaderCompileWorker.cpp b/Engine/Source/Programs/ShaderCompileWorker/Private/ShaderCompileWorker.cpp
--- a/Engine/Source/Programs/ShaderCompileWorker/Private/ShaderCompileWorker.cpp
+++ b/Engine/Source/Programs/ShaderCompileWorker/Private/ShaderCompileWorker.cpp
@@ -306,10 +306,19 @@ private:
 		FArchive* OutputFilePtr = NULL;
 		if (CommunicationMode == ThroughFile)
 		{
-			// Remove the input file so that it won't get processed more than once
-			if(!IFileManager::Get().Delete(*InputFilePath))
+			const double StartTime = FPlatformTime::Seconds();
+			bool bResult = false;
+
+			do 
+			{
+				// Remove the input file so that it won't get processed more than once
+				bResult = IFileManager::Get().Delete(*InputFilePath);
+			} 
+			while (!bResult && (FPlatformTime::Seconds() - StartTime < 2));
+
+			if (!bResult)
 			{
-				UE_LOG(LogShaders, Fatal,TEXT("Couldn't delete input file, is it readonly?"));
+				UE_LOG(LogShaders, Fatal,TEXT("Couldn't delete input file %s, is it readonly?"), *InputFilePath);
 			}
 
 #if PLATFORM_MAC			
@@ -323,10 +332,21 @@ private:
 			} while (IFileManager::Get().FileSize(*TempFilePath) != INDEX_NONE);
 
 			// Create the output file.
-			OutputFilePtr = IFileManager::Get().CreateFileWriter(*TempFilePath,FILEWRITE_NoFail);
+			OutputFilePtr = IFileManager::Get().CreateFileWriter(*TempFilePath,FILEWRITE_EvenIfReadOnly | FILEWRITE_NoFail);
 #else
-			// Create the output file.
-			OutputFilePtr = IFileManager::Get().CreateFileWriter(*OutputFilePath,FILEWRITE_NoFail);
+			const double StartTime2 = FPlatformTime::Seconds();
+
+			do 
+			{
+				// Create the output file.
+				OutputFilePtr = IFileManager::Get().CreateFileWriter(*OutputFilePath,FILEWRITE_EvenIfReadOnly);
+			} 
+			while (!OutputFilePtr && (FPlatformTime::Seconds() - StartTime2 < 2));
+			
+			if (!OutputFilePtr)
+			{
+				UE_LOG(LogShaders, Fatal,TEXT("Couldn't save output file %s"), *OutputFilePath);
+			}
 #endif
 		}
 		else
