Linux: fix ShaderCompileWorker

- from GitHub pull request #242.

diff --git a/Engine/Source/Programs/ShaderCompileWorker/Private/ShaderCompileWorker.cpp b/Engine/Source/Programs/ShaderCompileWorker/Private/ShaderCompileWorker.cpp
--- a/Engine/Source/Programs/ShaderCompileWorker/Private/ShaderCompileWorker.cpp
+++ b/Engine/Source/Programs/ShaderCompileWorker/Private/ShaderCompileWorker.cpp
@@ -81,7 +81,8 @@ static void ProcessCompilationJob(const FShaderCompilerInput& Input,FShaderCompi
 	Compiler->CompileShader(Input.ShaderFormat, Input, Output, WorkingDirectory);
 }
 
-static void VerifyResult(BOOL bResult, const TCHAR* InMessage = TEXT(""))
+#if PLATFORM_SUPPORTS_NAMED_PIPES
+static void VerifyResult(bool bResult, const TCHAR* InMessage = TEXT(""))
 {
 #if PLATFORM_WINDOWS
 	if (!bResult)
@@ -96,6 +97,7 @@ static void VerifyResult(BOOL bResult, const TCHAR* InMessage = TEXT(""))
 #endif // PLATFORM_WINDOWS
 	verify(bResult != 0);
 }
+#endif
 
 class FWorkLoop
 {
@@ -153,7 +155,7 @@ public:
 			// Close the output file.
 			delete OutputFilePtr;
 
-#if PLATFORM_MAC			
+#if PLATFORM_MAC || PLATFORM_LINUX
 			// Change the output file name to requested one
 			IFileManager::Get().Move(*OutputFilePath, *TempFilePath);
 #endif
@@ -196,7 +198,7 @@ private:
 
 	const FString InputFilePath;
 	const FString OutputFilePath;
-#if PLATFORM_MAC
+#if PLATFORM_MAC || PLATFORM_LINUX
 	FString TempFilePath;
 #endif
 
@@ -321,7 +323,7 @@ private:
 				UE_LOG(LogShaders, Fatal,TEXT("Couldn't delete input file %s, is it readonly?"), *InputFilePath);
 			}
 
-#if PLATFORM_MAC			
+#if PLATFORM_MAC || PLATFORM_LINUX
 			// To make sure that the process waiting for results won't read unfinished output file,
 			// we use a temp file name during compilation.
 			do
@@ -406,7 +408,7 @@ private:
 			}
 		}
 
-#if PLATFORM_MAC
+#if PLATFORM_MAC || PLATFORM_LINUX
 		if (!FPlatformMisc::IsDebuggerPresent() && ParentProcessId > 0)
 		{
 			// If the parent process is no longer running, exit
@@ -537,13 +539,13 @@ int32 GuardedMain(int32 argc, TCHAR* argv[])
 int32 GuardedMainWrapper(int32 ArgC, TCHAR* ArgV[], const TCHAR* CrashOutputFile)
 {
 	int32 ReturnCode = 0;
-#if !PLATFORM_MAC
+#if PLATFORM_WINDOWS
 	if (FPlatformMisc::IsDebuggerPresent())
 #endif
 	{
 		ReturnCode = GuardedMain(ArgC, ArgV);
 	}
-#if !PLATFORM_MAC
+#if PLATFORM_WINDOWS
 	else
 	{
 		__try
@@ -601,7 +603,8 @@ INT32_MAIN_INT32_ARGC_TCHAR_ARGV()
 {
 	if(ArgC < 6)
 	{
-		printf("ShaderCompileWorker.exe is called by UE4, it requires specific command like arguments.\n");
+		printf("ShaderCompileWorker is called by UE4, it requires specific command like arguments.\n");
+		return -1;
 	}
 
 	// Game exe can pass any number of parameters through with appGetSubprocessCommandline
diff --git a/Engine/Source/Programs/ShaderCompileWorker/ShaderCompileWorker.Build.cs b/Engine/Source/Programs/ShaderCompileWorker/ShaderCompileWorker.Build.cs
--- a/Engine/Source/Programs/ShaderCompileWorker/ShaderCompileWorker.Build.cs
+++ b/Engine/Source/Programs/ShaderCompileWorker/ShaderCompileWorker.Build.cs
@@ -17,6 +17,17 @@ public class ShaderCompileWorker : ModuleRules
 			}
 			);
 
+		if (Target.Platform == UnrealTargetPlatform.Linux)
+		{
+			PrivateDependencyModuleNames.AddRange(
+			new string[] {
+				"NetworkFile",
+				"PakFile",
+				"StreamingFile",
+				}
+			);
+		}
+
 		PrivateIncludePathModuleNames.AddRange(
 			new string[] {
 				"Launch",
