Fixed SCW on mac

diff --git a/Engine/Source/Programs/ShaderCompileWorker/Private/ShaderCompileWorker.cpp b/Engine/Source/Programs/ShaderCompileWorker/Private/ShaderCompileWorker.cpp
--- a/Engine/Source/Programs/ShaderCompileWorker/Private/ShaderCompileWorker.cpp
+++ b/Engine/Source/Programs/ShaderCompileWorker/Private/ShaderCompileWorker.cpp
@@ -544,7 +544,7 @@ int32 GuardedMain(int32 argc, TCHAR* argv[])
 	const bool bThroughNamedPipe = (InCommunicating == FString(TEXT("-communicatethroughnamedpipe")));
 	const bool bThroughNamedPipeOnce = (InCommunicating == FString(TEXT("-communicatethroughnamedpipeonce")));
 #else
-	const bool bThroughFile = GShaderCompileUseXGE;
+	const bool bThroughFile = true;
 	const bool bThroughNamedPipe = false;
 	const bool bThroughNamedPipeOnce = false;
 #endif
