SCW fixes
Corrected GLSL version for ios es2.
Fix a couple of out of bounds argv[6] reads.

diff --git a/Engine/Source/Developer/ShaderFormatOpenGL/Private/ShaderFormatOpenGL.cpp b/Engine/Source/Developer/ShaderFormatOpenGL/Private/ShaderFormatOpenGL.cpp
--- a/Engine/Source/Developer/ShaderFormatOpenGL/Private/ShaderFormatOpenGL.cpp
+++ b/Engine/Source/Developer/ShaderFormatOpenGL/Private/ShaderFormatOpenGL.cpp
@@ -157,7 +157,7 @@ public:
 		}
 		else if (Format == NAME_GLSL_ES2_IOS )
 		{
-			CompileShader_Windows_OGL(Input, Output, WorkingDirectory, GLSL_150_ES3_1);
+			CompileShader_Windows_OGL(Input, Output, WorkingDirectory, GLSL_ES2_IOS);
 		}
 		else if (Format == NAME_GLSL_150_ES3_1)
 		{
diff --git a/Engine/Source/Programs/ShaderCompileWorker/Private/ShaderCompileWorker.cpp b/Engine/Source/Programs/ShaderCompileWorker/Private/ShaderCompileWorker.cpp
--- a/Engine/Source/Programs/ShaderCompileWorker/Private/ShaderCompileWorker.cpp
+++ b/Engine/Source/Programs/ShaderCompileWorker/Private/ShaderCompileWorker.cpp
@@ -538,7 +538,7 @@ int32 GuardedMain(int32 argc, TCHAR* argv[])
 
 	LastCompileTime = FPlatformTime::Seconds();
 
-	FString InCommunicating = argv[6];
+	FString InCommunicating = (argc > 6) ? argv[6] : FString();
 #if PLATFORM_SUPPORTS_NAMED_PIPES
 	const bool bThroughFile = GShaderCompileUseXGE || (InCommunicating == FString(TEXT("-communicatethroughfile")));
 	const bool bThroughNamedPipe = (InCommunicating == FString(TEXT("-communicatethroughnamedpipe")));
@@ -562,7 +562,7 @@ int32 GuardedMainWrapper(int32 ArgC, TCHAR* ArgV[], const TCHAR* CrashOutputFile
 {
 	// We need to know whether we are using XGE now, in case an exception
 	// is thrown before we parse the command line inside GuardedMain.
-	GShaderCompileUseXGE = FCString::Strcmp(ArgV[6], TEXT("-xge")) == 0;
+	GShaderCompileUseXGE = (ArgC > 6) && FCString::Strcmp(ArgV[6], TEXT("-xge")) == 0;
 
 	int32 ReturnCode = 0;
 #if PLATFORM_WINDOWS
