Fix shadow variables in ShaderCompileWorker and UnrealPak

diff --git a/Engine/Source/Programs/ShaderCompileWorker/Private/ShaderCompileWorker.cpp b/Engine/Source/Programs/ShaderCompileWorker/Private/ShaderCompileWorker.cpp
--- a/Engine/Source/Programs/ShaderCompileWorker/Private/ShaderCompileWorker.cpp
+++ b/Engine/Source/Programs/ShaderCompileWorker/Private/ShaderCompileWorker.cpp
@@ -436,8 +436,8 @@ private:
 			// If the parent process is no longer running, exit
 			if (!FPlatformProcess::IsApplicationRunning(ParentProcessId))
 			{
-				FString InputFilePath = FString(WorkingDirectory) + InputFilename;
-				checkf(IFileManager::Get().FileSize(*InputFilePath) == INDEX_NONE, TEXT("Exiting due to the parent process no longer running and the input file is present!"));
+				FString FilePath = FString(WorkingDirectory) + InputFilename;
+				checkf(IFileManager::Get().FileSize(*FilePath) == INDEX_NONE, TEXT("Exiting due to the parent process no longer running and the input file is present!"));
 				UE_LOG(LogShaders, Log, TEXT("Parent process no longer running, exiting"));
 				FPlatformMisc::RequestExit(false);
 			}
@@ -457,7 +457,7 @@ private:
 		{
 			if (ParentProcessId > 0)
 			{
-				FString InputFilePath = FString(WorkingDirectory) + InputFilename;
+				FString FilePath = FString(WorkingDirectory) + InputFilename;
 
 				bool bParentStillRunning = true;
 				HANDLE ParentProcessHandle = OpenProcess(SYNCHRONIZE, false, ParentProcessId);
@@ -466,7 +466,7 @@ private:
 				{
 					if (!IsUsingNamedPipes())
 					{
-						checkf(IFileManager::Get().FileSize(*InputFilePath) == INDEX_NONE, TEXT("Exiting due to OpenProcess(ParentProcessId) failing and the input file is present!"));
+						checkf(IFileManager::Get().FileSize(*FilePath) == INDEX_NONE, TEXT("Exiting due to OpenProcess(ParentProcessId) failing and the input file is present!"));
 					}
 					UE_LOG(LogShaders, Log, TEXT("Couldn't OpenProcess, Parent process no longer running, exiting"));
 					FPlatformMisc::RequestExit(false);
@@ -481,7 +481,7 @@ private:
 					{
 						if (!IsUsingNamedPipes())
 						{
-							checkf(IFileManager::Get().FileSize(*InputFilePath) == INDEX_NONE, TEXT("Exiting due to WaitForSingleObject(ParentProcessHandle) signaling and the input file is present!"));
+							checkf(IFileManager::Get().FileSize(*FilePath) == INDEX_NONE, TEXT("Exiting due to WaitForSingleObject(ParentProcessHandle) signaling and the input file is present!"));
 						}
 						UE_LOG(LogShaders, Log, TEXT("WaitForSingleObject signaled, Parent process no longer running, exiting"));
 						FPlatformMisc::RequestExit(false);
diff --git a/Engine/Source/Programs/UnrealPak/Private/KeyGenerator.cpp b/Engine/Source/Programs/UnrealPak/Private/KeyGenerator.cpp
--- a/Engine/Source/Programs/UnrealPak/Private/KeyGenerator.cpp
+++ b/Engine/Source/Programs/UnrealPak/Private/KeyGenerator.cpp
@@ -256,24 +256,24 @@ void GeneratePrimeNumberTable(int64 MaxValue, const TCHAR* Filename)
 
 	UE_LOG(LogPakFile, Display, TEXT("Generating prime number table <= %lld: %s."), MaxValue, Filename);
 
-	FString PrimeTable(TEXT("// Copyright 1998-2015 Epic Games, Inc. All Rights Reserved.\nint256 PrimeTable[] = \n{\n\t2, "));
+	FString PrimeTableString(TEXT("// Copyright 1998-2015 Epic Games, Inc. All Rights Reserved.\nint256 PrimeTable[] = \n{\n\t2, "));
 	int64 PrimeCount = 1;
 	const double StartTime = FPlatformTime::Seconds();
 	for (int64 SmallNumber = 3; SmallNumber <= MaxValue; SmallNumber += 2)
 	{
 		if (IsPrime(SmallNumber, false))
 		{
-			PrimeTable += FString::Printf(TEXT("%lld, "), SmallNumber);
+			PrimeTableString += FString::Printf(TEXT("%lld, "), SmallNumber);
 			PrimeCount++;
 			if ((PrimeCount % 10) == 0)
 			{
-				PrimeTable += TEXT("\n\t");
+				PrimeTableString += TEXT("\n\t");
 			}
 		}
 	}
-	PrimeTable += TEXT("\n};\n");
+	PrimeTableString += TEXT("\n};\n");
 	UE_LOG(LogPakFile, Display, TEXT("Generated %lld primes in %.4lfs."), PrimeCount, FPlatformTime::Seconds() - StartTime);
-	FFileHelper::SaveStringToFile(PrimeTable, Filename);
+	FFileHelper::SaveStringToFile(PrimeTableString, Filename);
 }
 
 
