Shader compile worker file communication uses temporary file during writing with move to final location on all platforms, to hopefully increase reliability of file transfer on windows.

diff --git a/Engine/Source/Programs/ShaderCompileWorker/Private/ShaderCompileWorker.cpp b/Engine/Source/Programs/ShaderCompileWorker/Private/ShaderCompileWorker.cpp
--- a/Engine/Source/Programs/ShaderCompileWorker/Private/ShaderCompileWorker.cpp
+++ b/Engine/Source/Programs/ShaderCompileWorker/Private/ShaderCompileWorker.cpp
@@ -162,10 +162,8 @@ public:
 			// Close the output file.
 			delete OutputFilePtr;
 
-#if PLATFORM_MAC || PLATFORM_LINUX
 			// Change the output file name to requested one
 			IFileManager::Get().Move(*OutputFilePath, *TempFilePath);
-#endif
 
 #if PLATFORM_SUPPORTS_NAMED_PIPES
 			if (CommunicationMode == ThroughNamedPipeOnce || CommunicationMode == ThroughNamedPipe)
@@ -214,9 +212,7 @@ private:
 
 	const FString InputFilePath;
 	const FString OutputFilePath;
-#if PLATFORM_MAC || PLATFORM_LINUX
 	FString TempFilePath;
-#endif
 
 #if PLATFORM_SUPPORTS_NAMED_PIPES
 	FPlatformNamedPipe Pipe;
@@ -345,7 +341,6 @@ private:
 				}
 			}
 
-#if PLATFORM_MAC || PLATFORM_LINUX
 			// To make sure that the process waiting for results won't read unfinished output file,
 			// we use a temp file name during compilation.
 			do
@@ -355,23 +350,19 @@ private:
 				TempFilePath = WorkingDirectory + Guid.ToString();
 			} while (IFileManager::Get().FileSize(*TempFilePath) != INDEX_NONE);
 
-			// Create the output file.
-			OutputFilePtr = IFileManager::Get().CreateFileWriter(*TempFilePath,FILEWRITE_EvenIfReadOnly | FILEWRITE_NoFail);
-#else
 			const double StartTime2 = FPlatformTime::Seconds();
 
 			do 
 			{
 				// Create the output file.
-				OutputFilePtr = IFileManager::Get().CreateFileWriter(*OutputFilePath,FILEWRITE_EvenIfReadOnly);
+				OutputFilePtr = IFileManager::Get().CreateFileWriter(*TempFilePath,FILEWRITE_EvenIfReadOnly);
 			} 
 			while (!OutputFilePtr && (FPlatformTime::Seconds() - StartTime2 < 2));
 			
 			if (!OutputFilePtr)
 			{
-				UE_LOG(LogShaders, Fatal,TEXT("Couldn't save output file %s"), *OutputFilePath);
+				UE_LOG(LogShaders, Fatal,TEXT("Couldn't save output file %s"), *TempFilePath);
 			}
-#endif
 		}
 		else
 		{
diff --git a/Engine/Source/Runtime/Engine/Private/ShaderCompiler/ShaderCompiler.cpp b/Engine/Source/Runtime/Engine/Private/ShaderCompiler/ShaderCompiler.cpp
--- a/Engine/Source/Runtime/Engine/Private/ShaderCompiler/ShaderCompiler.cpp
+++ b/Engine/Source/Runtime/Engine/Private/ShaderCompiler/ShaderCompiler.cpp
@@ -716,7 +716,6 @@ void FShaderCompileThreadRunnable::WriteNewTasks()
 
 			const FString WorkingDirectory = Manager->AbsoluteShaderBaseWorkingDirectory + FString::FromInt(WorkerIndex);
 
-#if PLATFORM_MAC || PLATFORM_LINUX
 			// To make sure that the process waiting for input file won't try to read it until it's ready
 			// we use a temp file name during writing.
 			FString TransferFileName;
@@ -726,9 +725,6 @@ void FShaderCompileThreadRunnable::WriteNewTasks()
 				FPlatformMisc::CreateGuid(Guid);
 				TransferFileName = WorkingDirectory + Guid.ToString();
 			} while (IFileManager::Get().FileSize(*TransferFileName) != INDEX_NONE);
-#else
-			const FString TransferFileName = WorkingDirectory / TEXT("WorkerInputOnly.in");
-#endif
 
 			// Write out the file that the worker app is waiting for, which has all the information needed to compile the shader.
 			// 'Only' indicates that the worker should keep checking for more tasks after this one
@@ -764,11 +760,9 @@ void FShaderCompileThreadRunnable::WriteNewTasks()
 				DoWriteTasks(CurrentWorkerInfo.QueuedJobs, *TransferFile);
 				delete TransferFile;
 
-#if PLATFORM_MAC || PLATFORM_LINUX
 				// Change the transfer file name to proper one
 				FString ProperTransferFileName = WorkingDirectory / TEXT("WorkerInputOnly.in");
 				IFileManager::Get().Move(*ProperTransferFileName, *TransferFileName);
-#endif
 			}
 		}
 	}
