Build ShaderCompileWorker, UnrealCEFSubProcess and UnrealFrontend from RootEditor node to Tools node

diff --git a/Engine/Source/Programs/ShaderCompileWorker/ShaderCompileWorker.Target.cs b/Engine/Source/Programs/ShaderCompileWorker/ShaderCompileWorker.Target.cs
--- a/Engine/Source/Programs/ShaderCompileWorker/ShaderCompileWorker.Target.cs
+++ b/Engine/Source/Programs/ShaderCompileWorker/ShaderCompileWorker.Target.cs
@@ -72,10 +72,15 @@ public class ShaderCompileWorkerTarget : TargetRules
         // Linking against wer.lib/wer.dll causes XGE to bail when the worker is run on a Windows 8 machine, so turn this off.
         OutCPPEnvironmentConfiguration.Definitions.Add("ALLOW_WINDOWS_ERROR_REPORT_LIB=0");
 	}
-    public override bool GUBP_AlwaysBuildWithBaseEditor()
+
+    public override bool GUBP_AlwaysBuildWithTools(UnrealTargetPlatform InHostPlatform, out bool bInternalToolOnly, out bool SeparateNode, out bool CrossCompile)
     {
+        bInternalToolOnly = false;
+        SeparateNode = false;
+        CrossCompile = false;
         return true;
     }
+
     public override bool GUBP_NeedsPlatformSpecificDLLs()
     {
         return true;
diff --git a/Engine/Source/Programs/UnrealCEFSubProcess/UnrealCEFSubProcess.Target.cs b/Engine/Source/Programs/UnrealCEFSubProcess/UnrealCEFSubProcess.Target.cs
--- a/Engine/Source/Programs/UnrealCEFSubProcess/UnrealCEFSubProcess.Target.cs
+++ b/Engine/Source/Programs/UnrealCEFSubProcess/UnrealCEFSubProcess.Target.cs
@@ -69,10 +69,15 @@ public class UnrealCEFSubProcessTarget : TargetRules
 		// Disable logging, as the sub processes are spawned often and logging will just slow them down
 		OutCPPEnvironmentConfiguration.Definitions.Add("ALLOW_LOG_FILE=0");
 	}
-    public override bool GUBP_AlwaysBuildWithBaseEditor()
+
+    public override bool GUBP_AlwaysBuildWithTools(UnrealTargetPlatform InHostPlatform, out bool bInternalToolOnly, out bool SeparateNode, out bool CrossCompile)
     {
+        bInternalToolOnly = false;
+        SeparateNode = false;
+        CrossCompile = false;
         return true;
     }
+
     public override bool GUBP_NeedsPlatformSpecificDLLs()
     {
         return true;
diff --git a/Engine/Source/Programs/UnrealFrontend/UnrealFrontend.Target.cs b/Engine/Source/Programs/UnrealFrontend/UnrealFrontend.Target.cs
--- a/Engine/Source/Programs/UnrealFrontend/UnrealFrontend.Target.cs
+++ b/Engine/Source/Programs/UnrealFrontend/UnrealFrontend.Target.cs
@@ -49,10 +49,15 @@ public class UnrealFrontendTarget : TargetRules
 
 		OutLinkEnvironmentConfiguration.bHasExports = false;
 	}
-    public override bool GUBP_AlwaysBuildWithBaseEditor()
+
+    public override bool GUBP_AlwaysBuildWithTools(UnrealTargetPlatform InHostPlatform, out bool bInternalToolOnly, out bool SeparateNode, out bool CrossCompile)
     {
+        bInternalToolOnly = false;
+        SeparateNode = false;
+        CrossCompile = false;
         return true;
     }
+
     public override bool GUBP_NeedsPlatformSpecificDLLs()
     {
         return true;
