UE4 - Verify shader format versions due to mismatched SCW & Editor/Cooker

diff --git a/Engine/Shaders/ShaderVersion.usf b/Engine/Shaders/ShaderVersion.usf
--- a/Engine/Shaders/ShaderVersion.usf
+++ b/Engine/Shaders/ShaderVersion.usf
@@ -2,4 +2,4 @@
 // Each time the console command is executed it generates a new GUID. As this file is included
 // in common.usf (which should be included in any shader) it allows to invalidate the shader DDC.
 // 
-// GUID = 10DDEA3C48E7FD138E98039F9438498C
\ No newline at end of file
+// GUID = 10DDEA3C48E7FD138E98039F9438498D
\ No newline at end of file
diff --git a/Engine/Source/Programs/ShaderCompileWorker/Private/ShaderCompileWorker.cpp b/Engine/Source/Programs/ShaderCompileWorker/Private/ShaderCompileWorker.cpp
--- a/Engine/Source/Programs/ShaderCompileWorker/Private/ShaderCompileWorker.cpp
+++ b/Engine/Source/Programs/ShaderCompileWorker/Private/ShaderCompileWorker.cpp
@@ -13,19 +13,21 @@
 
 #define DEBUG_USING_CONSOLE	0
 
-const int32 ShaderCompileWorkerInputVersion = 2;
+const int32 ShaderCompileWorkerInputVersion = 3;
 const int32 ShaderCompileWorkerOutputVersion = 1;
 
 double LastCompileTime = 0.0;
 
 static bool GShaderCompileUseXGE = false;
+static bool GFailedDueToShaderFormatVersion = false;
+
 static void WriteXGESuccessFile(const TCHAR* WorkingDirectory)
 {
 	// To signal compilation completion, create a zero length file in the working directory.
 	delete IFileManager::Get().CreateFileWriter(*FString::Printf(TEXT("%s/Success"), WorkingDirectory), FILEWRITE_EvenIfReadOnly);
 }
 
-const TArray<const IShaderFormat*>& GetShaderFormats()
+static const TArray<const IShaderFormat*>& GetShaderFormats()
 {
 	static bool bInitialized = false;
 	static TArray<const IShaderFormat*> Results;
@@ -55,7 +57,7 @@ const TArray<const IShaderFormat*>& GetShaderFormats()
 	return Results;
 }
 
-const IShaderFormat* FindShaderFormat(FName Name)
+static const IShaderFormat* FindShaderFormat(FName Name)
 {
 	const TArray<const IShaderFormat*>& ShaderFormats = GetShaderFormats();	
 
@@ -88,6 +90,7 @@ static void ProcessCompilationJob(const FShaderCompilerInput& Input,FShaderCompi
 	Compiler->CompileShader(Input.ShaderFormat, Input, Output, WorkingDirectory);
 }
 
+
 class FWorkLoop
 {
 public:
@@ -95,7 +98,7 @@ public:
 	{
 		ThroughFile,
 	};
-	FWorkLoop(const TCHAR* ParentProcessIdText,const TCHAR* InWorkingDirectory,const TCHAR* InInputFilename,const TCHAR* InOutputFilename, ECommunicationMode InCommunicationMode)
+	FWorkLoop(const TCHAR* ParentProcessIdText,const TCHAR* InWorkingDirectory,const TCHAR* InInputFilename,const TCHAR* InOutputFilename, ECommunicationMode InCommunicationMode, TMap<FString, uint16>& InFormatVersionMap)
 	:	ParentProcessId(FCString::Atoi(ParentProcessIdText))
 	,	WorkingDirectory(InWorkingDirectory)
 	,	InputFilename(InInputFilename)
@@ -103,6 +106,7 @@ public:
 	,	CommunicationMode(InCommunicationMode)
 	,	InputFilePath(InCommunicationMode == ThroughFile ? (FString(InWorkingDirectory) + InInputFilename) : InInputFilename)
 	,	OutputFilePath(InCommunicationMode == ThroughFile ? (FString(InWorkingDirectory) + InOutputFilename) : InOutputFilename)
+	,	FormatVersionMap(InFormatVersionMap)
 	{
 	}
 
@@ -170,6 +174,7 @@ private:
 
 	const FString InputFilePath;
 	const FString OutputFilePath;
+	TMap<FString, uint16> FormatVersionMap;
 	FString TempFilePath;
 
 	/** Opens an input file, trying multiple times if necessary. */
@@ -196,6 +201,21 @@ private:
 		return InputFile;
 	}
 
+	void VerifyFormatVersions(TMap<FString, uint16>& ReceivedFormatVersionMap)
+	{
+		for (auto Pair : ReceivedFormatVersionMap)
+		{
+			auto* Found = FormatVersionMap.Find(Pair.Key);
+			if (Found)
+			{
+				GFailedDueToShaderFormatVersion = true;
+				FCString::Snprintf(GErrorExceptionDescription, sizeof(GErrorExceptionDescription), TEXT("Mismatched shader version for format %s; did you forget to build ShaderCompilerWorker?"), *Pair.Key, *Found, Pair.Value);
+
+				checkf(Pair.Value == *Found, TEXT("Exiting due to mismatched shader version for format %s, version %d from ShaderCompilerWorker, received %d! Did you forget to build ShaderCompilerWorker?"), *Pair.Key, *Found, Pair.Value);
+			}
+		}
+	}
+
 	void ProcessInputFromArchive(FArchive* InputFilePtr, TArray<FJobResult>& OutJobResults)
 	{
 		int32 NumBatches = 0;
@@ -203,7 +223,12 @@ private:
 		FArchive& InputFile = *InputFilePtr;
 		int32 InputVersion;
 		InputFile << InputVersion;
-		checkf(ShaderCompileWorkerInputVersion == InputVersion, TEXT("ShaderCompilerWorker expecting input version %d, got %d instead! Forgot to build ShaderCompilerWorker?"), ShaderCompileWorkerInputVersion, InputVersion);
+		checkf(ShaderCompileWorkerInputVersion == InputVersion, TEXT("Exiting due to ShaderCompilerWorker expecting input version %d, got %d instead! Did you forget to build ShaderCompilerWorker?"), ShaderCompileWorkerInputVersion, InputVersion);
+
+		TMap<FString, uint16> ReceivedFormatVersionMap;
+		InputFile << ReceivedFormatVersionMap;
+
+		VerifyFormatVersions(ReceivedFormatVersionMap);
 
 		InputFile << NumBatches;
 
@@ -408,6 +433,7 @@ int32 GuardedMain(int32 argc, TCHAR* argv[])
 	// We just enumerate the shader formats here for debugging.
 	const TArray<const class IShaderFormat*>& ShaderFormats = GetShaderFormats();
 	check(ShaderFormats.Num());
+	TMap<FString, uint16> FormatVersionMap;
 	for (int32 Index = 0; Index < ShaderFormats.Num(); Index++)
 	{
 		TArray<FName> OutFormats;
@@ -416,13 +442,15 @@ int32 GuardedMain(int32 argc, TCHAR* argv[])
 		for (int32 InnerIndex = 0; InnerIndex < OutFormats.Num(); InnerIndex++)
 		{
 			UE_LOG(LogShaders, Display, TEXT("Available Shader Format %s"), *OutFormats[InnerIndex].ToString());
+			uint16 Version = ShaderFormats[Index]->GetVersion(OutFormats[InnerIndex]);
+			FormatVersionMap.Add(OutFormats[InnerIndex].ToString(), Version);
 		}
 	}
 
 	LastCompileTime = FPlatformTime::Seconds();
 
 	FWorkLoop::ECommunicationMode Mode = FWorkLoop::ThroughFile;
-	FWorkLoop WorkLoop(argv[2], argv[1], argv[4], argv[5], Mode);
+	FWorkLoop WorkLoop(argv[2], argv[1], argv[4], argv[5], Mode, FormatVersionMap);
 
 	WorkLoop.Loop();
 
@@ -445,6 +473,8 @@ int32 GuardedMainWrapper(int32 ArgC, TCHAR* ArgV[], const TCHAR* CrashOutputFile
 #if PLATFORM_WINDOWS
 	else
 	{
+		// Don't want 32 dialogs popping up when SCW fails
+		GUseCrashReportClient = false;
 		__try
 		{
 			GIsGuarded = 1;
@@ -458,7 +488,7 @@ int32 GuardedMainWrapper(int32 ArgC, TCHAR* ArgV[], const TCHAR* CrashOutputFile
 			int32 OutputVersion = ShaderCompileWorkerOutputVersion;
 			OutputFile << OutputVersion;
 
-			int32 ErrorCode = 1;
+			int32 ErrorCode = GFailedDueToShaderFormatVersion ? 2 : 1;
 			OutputFile << ErrorCode;
 
 			// Note: Can't use FStrings here as SEH can't be used with destructors
diff --git a/Engine/Source/Runtime/Engine/Private/ShaderCompiler/ShaderCompiler.cpp b/Engine/Source/Runtime/Engine/Private/ShaderCompiler/ShaderCompiler.cpp
--- a/Engine/Source/Runtime/Engine/Private/ShaderCompiler/ShaderCompiler.cpp
+++ b/Engine/Source/Runtime/Engine/Private/ShaderCompiler/ShaderCompiler.cpp
@@ -18,6 +18,7 @@ DEFINE_LOG_CATEGORY(LogShaderCompilers);
 
 // Set to 1 to debug ShaderCompilerWorker.exe. Set a breakpoint in LaunchWorker() to get the cmd-line.
 #define DEBUG_SHADERCOMPILEWORKER 0
+
 // Default value comes from bPromptToRetryFailedShaderCompiles in BaseEngine.ini
 // This is set as a global variable to allow changing in the debugger even in release
 // For example if there are a lot of content shader compile errors you want to skip over without relaunching
@@ -108,11 +109,67 @@ namespace XGEConsoleVariables
 		ECVF_Default);
 }
 
+static const TArray<const IShaderFormat*>& GetShaderFormats()
+{
+	static bool bInitialized = false;
+	static TArray<const IShaderFormat*> Results;
+
+	if (!bInitialized)
+	{
+		bInitialized = true;
+		Results.Empty(Results.Num());
+
+		TArray<FName> Modules;
+		FModuleManager::Get().FindModules(SHADERFORMAT_MODULE_WILDCARD, Modules);
+
+		if (!Modules.Num())
+		{
+			UE_LOG(LogShaders, Error, TEXT("No target shader formats found!"));
+		}
+
+		for (int32 Index = 0; Index < Modules.Num(); Index++)
+		{
+			IShaderFormat* Format = FModuleManager::LoadModuleChecked<IShaderFormatModule>(Modules[Index]).GetShaderFormat();
+			if (Format != nullptr)
+			{
+				Results.Add(Format);
+			}
+		}
+	}
+	return Results;
+}
+
+static inline void GetFormatVersionMap(TMap<FString, uint16>& OutFormatVersionMap)
+{
+	if (OutFormatVersionMap.Num() == 0)
+	{
+		const TArray<const class IShaderFormat*>& ShaderFormats = GetShaderFormats();
+		check(ShaderFormats.Num());
+		for (int32 Index = 0; Index < ShaderFormats.Num(); Index++)
+		{
+			TArray<FName> OutFormats;
+			ShaderFormats[Index]->GetSupportedFormats(OutFormats);
+			check(OutFormats.Num());
+			for (int32 InnerIndex = 0; InnerIndex < OutFormats.Num(); InnerIndex++)
+			{
+				uint16 Version = ShaderFormats[Index]->GetVersion(OutFormats[InnerIndex]);
+				OutFormatVersionMap.Add(OutFormats[InnerIndex].ToString(), Version);
+			}
+		}
+	}
+}
+
 // Serialize Queued Job information
 static bool DoWriteTasks(const TArray<FShaderCompileJob*>& QueuedJobs, FArchive& TransferFile)
 {
-	int32 ShaderCompileWorkerInputVersion = 2;
+	int32 ShaderCompileWorkerInputVersion = 3;
 	TransferFile << ShaderCompileWorkerInputVersion;
+
+	static TMap<FString, uint16> FormatVersionMap;
+	GetFormatVersionMap(FormatVersionMap);
+
+	TransferFile << FormatVersionMap;
+
 	int32 NumBatches = QueuedJobs.Num();
 	TransferFile << NumBatches;
 
@@ -131,7 +188,22 @@ static void DoReadTaskResults(const TArray<FShaderCompileJob*>& QueuedJobs, FArc
 	const int32 OutputVersion = 1;
 	int32 ShaderCompileWorkerOutputVersion;
 	OutputFile << ShaderCompileWorkerOutputVersion;
-	checkf(ShaderCompileWorkerOutputVersion == OutputVersion, TEXT("Expecting ShaderCompilerWorker output version %d, got %d instead! Forgot to build ShaderCompilerWorker?"), ShaderCompileWorkerOutputVersion, OutputVersion);
+
+	if (ShaderCompileWorkerOutputVersion != OutputVersion)
+	{
+		FString Text = FString::Printf(TEXT("Expecting ShaderCompilerWorker output version %d, got %d instead! Forgot to build ShaderCompilerWorker?"), ShaderCompileWorkerOutputVersion, OutputVersion);
+		if (FPlatformProperties::SupportsWindowedMode())
+		{
+			UE_LOG(LogShaderCompilers, Error, TEXT("%s"), *Text);
+			FMessageDialog::Open(EAppMsgType::Ok, FText::FromString(Text));
+			FPlatformMisc::RequestExit(false);
+			return;
+		}
+		else
+		{
+			UE_LOG(LogShaderCompilers, Fatal, TEXT("%s"), *Text);
+		}
+	}
 
 	int32 ErrorCode;
 	OutputFile << ErrorCode;
@@ -143,20 +215,33 @@ static void DoReadTaskResults(const TArray<FShaderCompileJob*>& QueuedJobs, FArc
 	OutputFile << ExceptionInfoLength;
 
 	// Worker crashed
-	if (ErrorCode == 1)
+	if (ErrorCode == 1 || ErrorCode == 2)
 	{
-		TCHAR* Callstack = new TCHAR[CallstackLength + 1];
-		OutputFile.Serialize(Callstack, CallstackLength * sizeof(TCHAR));
+		TArray<TCHAR> Callstack;
+		Callstack.AddUninitialized(CallstackLength + 1);
+		OutputFile.Serialize(Callstack.GetData(), CallstackLength * sizeof(TCHAR));
 		Callstack[CallstackLength] = 0;
 
-		TCHAR* ExceptionInfo = new TCHAR[ExceptionInfoLength + 1];
-		OutputFile.Serialize(ExceptionInfo, ExceptionInfoLength * sizeof(TCHAR));
+		TArray<TCHAR> ExceptionInfo;
+		ExceptionInfo.AddUninitialized(ExceptionInfoLength + 1);
+		OutputFile.Serialize(ExceptionInfo.GetData(), ExceptionInfoLength * sizeof(TCHAR));
 		ExceptionInfo[ExceptionInfoLength] = 0;
 
-		UE_LOG(LogShaderCompilers, Fatal, TEXT("ShaderCompileWorker crashed! \n %s \n %s"), ExceptionInfo, Callstack);
+		if (ErrorCode == 2)
+		{
+			if (FPlatformProperties::SupportsWindowedMode())
+			{
+				UE_LOG(LogShaderCompilers, Error, TEXT("%s\n%s"), ExceptionInfo.GetData(), Callstack.GetData());
+				FMessageDialog::Open(EAppMsgType::Ok, FText::FromString(ExceptionInfo.GetData()));
+				FPlatformMisc::RequestExit(false);
+			}
+			else
+			{
+				UE_LOG(LogShaderCompilers, Fatal, TEXT("%s"), Callstack.GetData());
+			}
+		}
 
-		delete [] Callstack;
-		delete [] ExceptionInfo;
+		UE_LOG(LogShaderCompilers, Fatal, TEXT("ShaderCompileWorker crashed! \n %s \n %s"), ExceptionInfo.GetData(), Callstack.GetData());
 	}
 
 	int32 NumJobs;
@@ -872,7 +957,7 @@ FProcHandle FShaderCompilingManager::LaunchWorker(const FString& WorkingDirector
 	// Note: Set breakpoint here and launch the shadercompileworker with WorkerParameters a cmd-line
 	const TCHAR* WorkerParametersText = *WorkerParameters;
 	FPlatformMisc::LowLevelOutputDebugStringf(TEXT("Launching shader compile worker w/ WorkerParameters\n\t%s\n"), WorkerParametersText);
-	return FProcHandle(17);
+	return FProcHandle();
 #else
 #if UE_BUILD_DEBUG && PLATFORM_LINUX
 	FPlatformMisc::LowLevelOutputDebugStringf(TEXT("Launching shader compile worker:\n\t%s\n"), *WorkerParameters);
