UE4 - SCW - Do not add bogus error description when not needed

diff --git a/Engine/Source/Programs/ShaderCompileWorker/Private/ShaderCompileWorker.cpp b/Engine/Source/Programs/ShaderCompileWorker/Private/ShaderCompileWorker.cpp
--- a/Engine/Source/Programs/ShaderCompileWorker/Private/ShaderCompileWorker.cpp
+++ b/Engine/Source/Programs/ShaderCompileWorker/Private/ShaderCompileWorker.cpp
@@ -209,9 +209,12 @@ private:
 			if (Found)
 			{
 				GFailedDueToShaderFormatVersion = true;
-				FCString::Snprintf(GErrorExceptionDescription, sizeof(GErrorExceptionDescription), TEXT("Mismatched shader version for format %s; did you forget to build ShaderCompilerWorker?"), *Pair.Key, *Found, Pair.Value);
+				if (Pair.Value != *Found)
+				{
+					FCString::Snprintf(GErrorExceptionDescription, sizeof(GErrorExceptionDescription), TEXT("Mismatched shader version for format %s; did you forget to build ShaderCompilerWorker?"), *Pair.Key, *Found, Pair.Value);
 
-				checkf(Pair.Value == *Found, TEXT("Exiting due to mismatched shader version for format %s, version %d from ShaderCompilerWorker, received %d! Did you forget to build ShaderCompilerWorker?"), *Pair.Key, *Found, Pair.Value);
+					checkf(Pair.Value == *Found, TEXT("Exiting due to mismatched shader version for format %s, version %d from ShaderCompilerWorker, received %d! Did you forget to build ShaderCompilerWorker?"), *Pair.Key, *Found, Pair.Value);
+				}
 			}
 		}
 	}
