Fix ShaderCompileWorker building as a dependency of UE4Game in Win32.

diff --git a/Engine/Source/Programs/ShaderCompileWorker/ShaderCompileWorker.Target.cs b/Engine/Source/Programs/ShaderCompileWorker/ShaderCompileWorker.Target.cs
--- a/Engine/Source/Programs/ShaderCompileWorker/ShaderCompileWorker.Target.cs
+++ b/Engine/Source/Programs/ShaderCompileWorker/ShaderCompileWorker.Target.cs
@@ -17,7 +17,7 @@ public class ShaderCompileWorkerTarget : TargetRules
 
 	public override bool GetSupportedPlatforms(ref List<UnrealTargetPlatform> OutPlatforms)
 	{
-		return UnrealBuildTool.UnrealBuildTool.GetAllDesktopPlatforms(ref OutPlatforms, false);
+		return UnrealBuildTool.UnrealBuildTool.GetAllEditorPlatforms(ref OutPlatforms, false);
 	}
 
 	public override bool ShouldCompileMonolithic(UnrealTargetPlatform InPlatform, UnrealTargetConfiguration InConfiguration)
