Fix a shadowed variable.

diff --git a/Engine/Source/Programs/ShaderCompileWorker/Private/ShaderCompileWorker.cpp b/Engine/Source/Programs/ShaderCompileWorker/Private/ShaderCompileWorker.cpp
--- a/Engine/Source/Programs/ShaderCompileWorker/Private/ShaderCompileWorker.cpp
+++ b/Engine/Source/Programs/ShaderCompileWorker/Private/ShaderCompileWorker.cpp
@@ -328,8 +328,10 @@ private:
 		ProcessCompilationJob(CompilerInputs[0], FirstCompilerOutput, WorkingDirectory);
 
 		// Serialize the job's output.
-		FJobResult& JobResult = *new(PipelineJob.SingleJobs) FJobResult;
-		JobResult.CompilerOutput = FirstCompilerOutput;
+		{
+			FJobResult& JobResult = *new(PipelineJob.SingleJobs) FJobResult;
+			JobResult.CompilerOutput = FirstCompilerOutput;
+		}
 
 		bool bEnableRemovingUnused = true;
 
