Make sure we unload modules when shutting down SCW. This solves a crash at exit where FDDCCleanup would try to use IFileManager.Get() after the file manager was destroyed

diff --git a/Engine/Source/Programs/ShaderCompileWorker/Private/ShaderCompileWorker.cpp b/Engine/Source/Programs/ShaderCompileWorker/Private/ShaderCompileWorker.cpp
--- a/Engine/Source/Programs/ShaderCompileWorker/Private/ShaderCompileWorker.cpp
+++ b/Engine/Source/Programs/ShaderCompileWorker/Private/ShaderCompileWorker.cpp
@@ -892,6 +892,7 @@ static int32 GuardedMainWrapper(int32 ArgC, TCHAR* ArgV[], const TCHAR* CrashOut
 #endif
 
 	FEngineLoop::AppPreExit();
+	FModuleManager::Get().UnloadModulesAtShutdown();
 	FEngineLoop::AppExit();
 
 	return ReturnCode;
