VR - Missed file from Lumin platform split

diff --git a/Engine/Source/Programs/ShaderCompileWorker/Private/ShaderCompileWorker.cpp b/Engine/Source/Programs/ShaderCompileWorker/Private/ShaderCompileWorker.cpp
--- a/Engine/Source/Programs/ShaderCompileWorker/Private/ShaderCompileWorker.cpp
+++ b/Engine/Source/Programs/ShaderCompileWorker/Private/ShaderCompileWorker.cpp
@@ -701,12 +701,14 @@ static FName NAME_GLSL_310_ES_EXT(TEXT("GLSL_310_ES_EXT"));
 static FName NAME_SF_METAL_SM5(TEXT("SF_METAL_SM5"));
 static FName NAME_VULKAN_ES3_1_ANDROID(TEXT("SF_VULKAN_ES31_ANDROID"));
 static FName NAME_VULKAN_ES3_1_ANDROID_NOUB(TEXT("SF_VULKAN_ES31_ANDROID_NOUB"));
+static FName NAME_VULKAN_ES3_1_LUMIN(TEXT("SF_VULKAN_ES31_LUMIN"));
 static FName NAME_VULKAN_ES3_1(TEXT("SF_VULKAN_ES31"));
 static FName NAME_VULKAN_ES3_1_NOUB(TEXT("SF_VULKAN_ES31_NOUB"));
 static FName NAME_VULKAN_SM4_NOUB(TEXT("SF_VULKAN_SM4_NOUB"));
 static FName NAME_VULKAN_SM4(TEXT("SF_VULKAN_SM4"));
 static FName NAME_VULKAN_SM5_NOUB(TEXT("SF_VULKAN_SM5_NOUB"));
 static FName NAME_VULKAN_SM5(TEXT("SF_VULKAN_SM5"));
+static FName NAME_VULKAN_SM5_LUMIN(TEXT("SF_VULKAN_SM5_LUMIN"));
 static FName NAME_SF_METAL_SM5_NOTESS(TEXT("SF_METAL_SM5_NOTESS"));
 static FName NAME_SF_METAL_MACES3_1(TEXT("SF_METAL_MACES3_1"));
 static FName NAME_GLSL_ES3_1_ANDROID(TEXT("GLSL_ES3_1_ANDROID"));
@@ -733,8 +735,10 @@ static EShaderPlatform FormatNameToEnum(FName ShaderFormat)
 	if (ShaderFormat == NAME_SF_METAL_SM5)		return SP_METAL_SM5;
 	if (ShaderFormat == NAME_VULKAN_SM4)			return SP_VULKAN_SM4;
 	if (ShaderFormat == NAME_VULKAN_SM5)			return SP_VULKAN_SM5;
+	if (ShaderFormat == NAME_VULKAN_SM5_LUMIN)		return SP_VULKAN_SM5_LUMIN;
 	if (ShaderFormat == NAME_VULKAN_ES3_1_ANDROID)	return SP_VULKAN_ES3_1_ANDROID;
 	if (ShaderFormat == NAME_VULKAN_ES3_1_ANDROID_NOUB)	return SP_VULKAN_ES3_1_ANDROID;
+	if (ShaderFormat == NAME_VULKAN_ES3_1_LUMIN)		return SP_VULKAN_ES3_1_LUMIN;
 	if (ShaderFormat == NAME_VULKAN_ES3_1)			return SP_VULKAN_PCES3_1;
 	if (ShaderFormat == NAME_VULKAN_ES3_1_NOUB)	return SP_VULKAN_PCES3_1;
 	if (ShaderFormat == NAME_VULKAN_SM4_NOUB)		return SP_VULKAN_SM4;
