Optimize shader compilation through XGE (bringing CL 11084076)

  - Avoid directory scanning for .ini files by restoring already initialized config cache state on the workers
  - Avoid directory scanning for modules by restoring already initialized modulemanager state on the workers
  - Avoid directory scanning for external profilers DLLs by disabling the option in the build config
  - 20s -> 319ms of FEngineLoop::PreInit which was caused by directory scan through XGE remote filesystem
  - 5% -> 96% efficiency when computing the effective work against process total time for remotely built shaders
  - 5m36 ->1m26s to run "recompileshaders all" console command including waiting on async built shaders

diff --git a/Engine/Source/Programs/ShaderCompileWorker/Private/ShaderCompileWorker.cpp b/Engine/Source/Programs/ShaderCompileWorker/Private/ShaderCompileWorker.cpp
--- a/Engine/Source/Programs/ShaderCompileWorker/Private/ShaderCompileWorker.cpp
+++ b/Engine/Source/Programs/ShaderCompileWorker/Private/ShaderCompileWorker.cpp
@@ -61,9 +61,9 @@ static void OnXGEJobCompleted(const TCHAR* WorkingDirectory)
 {
 	if (GXGEMode == EXGEMode::Xml)
 	{
-	// To signal compilation completion, create a zero length file in the working directory.
+		// To signal compilation completion, create a zero length file in the working directory.
 		// This is only required in Xml mode.
-	delete IFileManager::Get().CreateFileWriter(*FString::Printf(TEXT("%s/Success"), WorkingDirectory), FILEWRITE_EvenIfReadOnly);
+		delete IFileManager::Get().CreateFileWriter(*FString::Printf(TEXT("%s/Success"), WorkingDirectory), FILEWRITE_EvenIfReadOnly);
 	}
 }
 
@@ -860,11 +860,94 @@ static void DirectCompile(const TArray<const class IShaderFormat*>& ShaderFormat
  */
 static int32 GuardedMain(int32 argc, TCHAR* argv[], bool bDirectMode)
 {
-	GEngineLoop.PreInit(argc, argv, TEXT("-NOPACKAGECACHE -ReduceThreadUsage"));
+	FString ExtraCmdLine = TEXT("-NOPACKAGECACHE -ReduceThreadUsage -cpuprofilertrace");
+
+	// When executing tasks remotely through XGE, enumerating files requires tcp/ip round-trips with
+	// the initiator, which can slow down engine initialization quite drastically.
+	// The idea here is to save the Ini and Modules manager state and reuse them on the workers
+	// to avoid all those directory enumeration during engine init.
+	FString IniBootstrapFilename;
+	FString ModulesBootstrapFilename;
+
+	if (IsUsingXGE())
+	{
+		// Tie the bootstrap filenames to the xge job id to refresh bootstraps state every time a new build starts
+		// This allows the ini/modules and shadercompilerworker binaries to change between builds.
+		FGuid XGJobID;
+		if (FGuid::Parse(FPlatformMisc::GetEnvironmentVariable(TEXT("xgJobID")), XGJobID))
+		{
+			FString XGJobIDString = XGJobID.ToString(EGuidFormats::DigitsWithHyphens);
+			IniBootstrapFilename = FString::Printf(TEXT("%s/Bootstrap-%s.inis"), argv[1], *XGJobIDString);
+			ModulesBootstrapFilename = FString::Printf(TEXT("%s/Bootstrap-%s.modules"), argv[1], *XGJobIDString);
+
+			ExtraCmdLine.Appendf(TEXT(" -IniBootstrap=\"%s\" -ModulesBootstrap=\"%s\""), *IniBootstrapFilename, *ModulesBootstrapFilename);
+
+			// Use Windows API directly because required CreateFile flags are not supported by our current OS abstraction
+#if PLATFORM_WINDOWS
+			// This is advantageous to have only a single worker do the init work instead of having all workers
+			// do a stampede of the initiator's machine all trying to enumerate directories at the same time.
+			// I've seen incoming TCP connections going through the roof (350 connections for 150 virtual CPUs)
+			// coming from workers doing all the same directory enumerations.
+			// This is not strictly required, but will improve performance when successful.
+			// Most likely a local worker will win the race and do a fast init.
+			FString MutexFilename = FString::Printf(TEXT("%s/Bootstrap-%s.mutex"), argv[1], *XGJobIDString);
+
+			// We need to implement a mutex scheme through a file for it to work with XGE's file virtualization layer.
+			// The first process to successfully create this file will have the honor of doing the complete initialization.
+			HANDLE MutexHandle =
+				CreateFileW(
+					*MutexFilename,
+					GENERIC_WRITE,
+					0,
+					nullptr,
+					CREATE_NEW,
+					FILE_ATTRIBUTE_NORMAL,
+					nullptr);
+
+			if (MutexHandle != INVALID_HANDLE_VALUE)
+			{
+				// We won the race, proceed to initialization.
+				CloseHandle(MutexHandle);
+			}
+			else
+			{
+				// Wait until the race winner writes the last bootstrap file
+				// Due to a bug in XGE, some workers might never see the new file appear, we must proceed after some timeout value.
+				for (int32 Index = 0; Index < 10 && !FPaths::FileExists(ModulesBootstrapFilename); ++Index)
+				{
+					Sleep(100);
+				}
+			}
+#endif
+		}
+	}
+
+	GEngineLoop.PreInit(argc, argv, *ExtraCmdLine);
 #if DEBUG_USING_CONSOLE
 	GLogConsole->Show( true );
 #endif
 
+	auto AtomicSave = 
+		[](const FString& Filename, TFunctionRef<void (const FString& TmpFile)> SaveFunction)
+		{
+			if (!Filename.IsEmpty() && !FPaths::FileExists(Filename))
+			{
+				// Use a tmp file for atomic publication and avoid reading incomplete state from other workers
+				FString TmpFile = FString::Printf(TEXT("%s-%s"), *Filename, *FGuid::NewGuid().ToString());
+				SaveFunction(TmpFile);
+				const bool bReplace = false;
+				const bool bDoNotRetryOrError = true;
+				const bool bEvenIfReadOnly = false;
+				const bool bAttributes = false;
+				IFileManager::Get().Move(*Filename, *TmpFile, bReplace, bEvenIfReadOnly, bAttributes, bDoNotRetryOrError);
+				// In case this process lost the race and wasn't able to move the file, discard the tmp file.
+				IFileManager::Get().Delete(*TmpFile);
+			}
+		};
+
+	AtomicSave(IniBootstrapFilename,     [](const FString& TmpFile) { GConfig->SaveCurrentStateForBootstrap(*TmpFile); });
+	AtomicSave(ModulesBootstrapFilename, [](const FString& TmpFile) { FModuleManager::Get().SaveCurrentStateForBootstrap(*TmpFile); });
+
 	// We just enumerate the shader formats here for debugging.
 	const TArray<const class IShaderFormat*>& ShaderFormats = GetShaderFormats();
 	check(ShaderFormats.Num());
diff --git a/Engine/Source/Programs/ShaderCompileWorker/ShaderCompileWorker.Build.cs b/Engine/Source/Programs/ShaderCompileWorker/ShaderCompileWorker.Build.cs
--- a/Engine/Source/Programs/ShaderCompileWorker/ShaderCompileWorker.Build.cs
+++ b/Engine/Source/Programs/ShaderCompileWorker/ShaderCompileWorker.Build.cs
@@ -14,7 +14,8 @@ public class ShaderCompileWorker : ModuleRules
 				"RenderCore",
 				"SandboxFile",
 				"TargetPlatform",
-				"ApplicationCore"
+				"ApplicationCore",
+				"TraceLog"
 			});
 
 		if (Target.Platform == UnrealTargetPlatform.Linux)
diff --git a/Engine/Source/Programs/ShaderCompileWorker/ShaderCompileWorker.Target.cs b/Engine/Source/Programs/ShaderCompileWorker/ShaderCompileWorker.Target.cs
--- a/Engine/Source/Programs/ShaderCompileWorker/ShaderCompileWorker.Target.cs
+++ b/Engine/Source/Programs/ShaderCompileWorker/ShaderCompileWorker.Target.cs
@@ -50,7 +50,7 @@ public class ShaderCompileWorkerTarget : TargetRules
 		bUseMallocProfiler = false;
 
 		// Force all shader formats to be built and included.
-        bForceBuildShaderFormats = true;
+		bForceBuildShaderFormats = true;
 
 		// ShaderCompileWorker is a console application, not a Windows app (sets entry point to main(), instead of WinMain())
 		bIsBuildingConsoleApplication = true;
@@ -58,7 +58,10 @@ public class ShaderCompileWorkerTarget : TargetRules
 		// Disable logging, as the workers are spawned often and logging will just slow them down
 		GlobalDefinitions.Add("ALLOW_LOG_FILE=0");
 
-        // Linking against wer.lib/wer.dll causes XGE to bail when the worker is run on a Windows 8 machine, so turn this off.
-        GlobalDefinitions.Add("ALLOW_WINDOWS_ERROR_REPORT_LIB=0");
+		// Linking against wer.lib/wer.dll causes XGE to bail when the worker is run on a Windows 8 machine, so turn this off.
+		GlobalDefinitions.Add("ALLOW_WINDOWS_ERROR_REPORT_LIB=0");
+
+		// Disable external profiling in ShaderCompiler to improve startup time
+		GlobalDefinitions.Add("UE_EXTERNAL_PROFILING_ENABLED=0");
 	}
 }
diff --git a/Engine/Source/Runtime/Core/Private/Misc/ConfigCacheIni.cpp b/Engine/Source/Runtime/Core/Private/Misc/ConfigCacheIni.cpp
--- a/Engine/Source/Runtime/Core/Private/Misc/ConfigCacheIni.cpp
+++ b/Engine/Source/Runtime/Core/Private/Misc/ConfigCacheIni.cpp
@@ -18,6 +18,8 @@
 #include "Misc/DefaultValueHelper.h"
 #include "Misc/ConfigManifest.h"
 #include "Misc/DataDrivenPlatformInfoRegistry.h"
+#include "Serialization/MemoryReader.h"
+#include "Serialization/MemoryWriter.h"
 #include "Misc/StringBuilder.h"
 
 #if WITH_EDITOR
@@ -328,6 +330,13 @@ bool FConfigSection::operator!=( const FConfigSection& Other ) const
 	return ! (FConfigSection::operator==(Other));
 }
 
+FArchive& operator<<(FArchive& Ar, FConfigSection& ConfigSection)
+{
+	Ar << static_cast<FConfigSection::Super&>(ConfigSection);
+	Ar << ConfigSection.ArrayOfStructKeys;
+	return Ar;
+}
+
 // Pull out a property from a Struct property, StructKeyMatch should be in the form "MyProp=". This reduces
 // memory allocations for each attempted match
 static void ExtractPropertyValue(const FString& FullStructValue, const FString& StructKeyMatch, FString& Out)
@@ -3749,8 +3758,73 @@ FString FConfigCacheIni::GetDestIniFilename(const TCHAR* BaseIniName, const TCHA
 	return IniFilename;
 }
 
+void FConfigCacheIni::SaveCurrentStateForBootstrap(const TCHAR* Filename)
+{
+	TArray<uint8> FileContent;
+	{
+		// Use FMemoryWriter because FileManager::CreateFileWriter doesn't serialize FName as string and is not overridable
+		FMemoryWriter MemoryWriter(FileContent, true);
+		SerializeStateForBootstrap_Impl(MemoryWriter);
+	}
+
+	FFileHelper::SaveArrayToFile(FileContent, Filename);
+}
+
+FArchive& operator<<(FArchive& Ar, FConfigCacheIni& ConfigCacheIni)
+{
+	Ar << static_cast<FConfigCacheIni::Super&>(ConfigCacheIni);
+	Ar << ConfigCacheIni.bAreFileOperationsDisabled;
+	Ar << ConfigCacheIni.bIsReadyForUse;
+	Ar << ConfigCacheIni.Type;
+	return Ar;
+}
+
+void FConfigCacheIni::SerializeStateForBootstrap_Impl(FArchive& Ar)
+{
+	// This implementation is meant to stay private and be used for 
+	// bootstrapping another processes' config cache with a serialized state.
+	// It doesn't include any versioning as it is used with the
+	// the same binary executable for both the parent and 
+	// children processes. It also takes care of saving/restoring
+	// global ini variables.
+	Ar << *this;
+	Ar << GEditorIni;
+	Ar << GEditorKeyBindingsIni;
+	Ar << GEditorLayoutIni;
+	Ar << GEditorSettingsIni;
+	Ar << GEditorPerProjectIni;
+	Ar << GCompatIni;
+	Ar << GLightmassIni;
+	Ar << GScalabilityIni;
+	Ar << GHardwareIni;
+	Ar << GInputIni;
+	Ar << GGameIni;
+	Ar << GGameUserSettingsIni;
+	Ar << GRuntimeOptionsIni;
+}
+
 void FConfigCacheIni::InitializeConfigSystem()
 {
+	// Bootstrap the Ini config cache
+	FString IniBootstrapFilename;
+	if (FParse::Value( FCommandLine::Get(), TEXT("IniBootstrap="), IniBootstrapFilename))
+	{
+		TArray<uint8> FileContent;
+		if (FFileHelper::LoadFileToArray(FileContent, *IniBootstrapFilename, FILEREAD_Silent))
+		{
+			FMemoryReader MemoryReader(FileContent, true);
+			GConfig = new FConfigCacheIni(EConfigCacheType::Temporary);
+			GConfig->SerializeStateForBootstrap_Impl(MemoryReader);
+			GConfig->bIsReadyForUse = true;
+			FCoreDelegates::ConfigReadyForUse.Broadcast();
+			return;
+		}
+		else
+		{
+			UE_LOG(LogInit, Display, TEXT("Unable to bootstrap from archive %s, will fallback on normal initialization"), *IniBootstrapFilename);
+		}
+	}
+
 	// Perform any upgrade we need before we load any configuration files
 	FConfigManifest::UpgradeFromPreviousVersions();
 
@@ -4046,6 +4120,34 @@ FString FConfigCacheIni::GetGameUserSettingsDir()
 	return ConfigDir;
 }
 
+FArchive& operator<<(FArchive& Ar, FConfigFile& ConfigFile)
+{
+	bool bHasSourceConfigFile = ConfigFile.SourceConfigFile != nullptr;
+
+	Ar << static_cast<FConfigFile::Super&>(ConfigFile);
+	Ar << ConfigFile.Dirty;
+	Ar << ConfigFile.NoSave;
+	Ar << ConfigFile.Name;
+	Ar << ConfigFile.SourceIniHierarchy;
+	Ar << ConfigFile.SourceEngineConfigDir;
+	Ar << bHasSourceConfigFile;
+	if (bHasSourceConfigFile)
+	{
+		// Handle missing instance for the loading case
+		if (ConfigFile.SourceConfigFile == nullptr)
+		{
+			ConfigFile.SourceConfigFile = new FConfigFile();
+		}
+
+		Ar << *ConfigFile.SourceConfigFile;
+	}
+	Ar << ConfigFile.SourceProjectConfigDir;
+	Ar << ConfigFile.CacheKey;
+	Ar << ConfigFile.PerObjectConfigArrayOfStructKeys;
+
+	return Ar;
+}
+
 void FConfigFile::UpdateSections(const TCHAR* DiskFilename, const TCHAR* IniRootName/*=nullptr*/, const TCHAR* OverridePlatform/*=nullptr*/)
 {
 	// Since we don't want any modifications to other sections, we manually process the file, not read into sections, etc
diff --git a/Engine/Source/Runtime/Core/Private/Modules/ModuleManager.cpp b/Engine/Source/Runtime/Core/Private/Modules/ModuleManager.cpp
--- a/Engine/Source/Runtime/Core/Private/Modules/ModuleManager.cpp
+++ b/Engine/Source/Runtime/Core/Private/Modules/ModuleManager.cpp
@@ -12,6 +12,9 @@
 #include "Misc/ScopeLock.h"
 #include "Misc/DataDrivenPlatformInfoRegistry.h"
 #include "Serialization/LoadTimeTrace.h"
+#include "Misc/FileHelper.h"
+#include "Serialization/MemoryReader.h"
+#include "Serialization/MemoryWriter.h"
 
 DEFINE_LOG_CATEGORY_STATIC(LogModuleManager, Log, All);
 
@@ -71,8 +74,29 @@ FModuleManager& FModuleManager::Get()
 
 FModuleManager::FModuleManager()
 	: bCanProcessNewlyLoadedObjects(false)
+	, bExtraBinarySearchPathsAdded(false)
 {
 	check(IsInGameThread());
+
+#if !IS_MONOLITHIC
+	// Modules bootstrapping is useful to avoid costly directory enumeration by reloading
+	// a serialized state of the module manager. Can only be used when run in the exact
+	// same context multiple times (i.e. starting multiple shader compile workers)
+	FString ModulesBootstrapFilename;
+	if (FParse::Value(FCommandLine::Get(), TEXT("ModulesBootstrap="), ModulesBootstrapFilename))
+	{
+		TArray<uint8> FileContent;
+		if (FFileHelper::LoadFileToArray(FileContent, *ModulesBootstrapFilename, FILEREAD_Silent))
+		{
+			FMemoryReader MemoryReader(FileContent, true);
+			SerializeStateForBootstrap_Impl(MemoryReader);
+		}
+		else
+		{
+			UE_LOG(LogModuleManager, Display, TEXT("Unable to bootstrap from archive %s, will fallback on normal initialization"), *ModulesBootstrapFilename);
+		}
+	}
+#endif
 }
 
 FModuleManager::~FModuleManager()
@@ -908,6 +932,60 @@ bool FModuleManager::HasAnyOverridenModuleFilename() const
 	}
 	return false;
 }
+
+void FModuleManager::SaveCurrentStateForBootstrap(const TCHAR* Filename)
+{
+	TArray<uint8> FileContent;
+	{
+		// Use FMemoryWriter because FileManager::CreateFileWriter doesn't serialize FName as string and is not overridable
+		FMemoryWriter MemoryWriter(FileContent, true);
+		FModuleManager::Get().SerializeStateForBootstrap_Impl(MemoryWriter);
+	}
+
+	FFileHelper::SaveArrayToFile(FileContent, Filename);
+}
+
+FArchive& operator<<(FArchive& Ar, FModuleManager& ModuleManager)
+{
+	Ar << ModuleManager.ModulePathsCache;
+	Ar << ModuleManager.PendingEngineBinariesDirectories;
+	Ar << ModuleManager.PendingGameBinariesDirectories;
+	Ar << ModuleManager.EngineBinariesDirectories;
+	Ar << ModuleManager.GameBinariesDirectories;
+	Ar << ModuleManager.bExtraBinarySearchPathsAdded;
+	Ar << ModuleManager.BuildId;
+	return Ar;
+}
+
+void FModuleManager::SerializeStateForBootstrap_Impl(FArchive& Ar)
+{
+	// This implementation is meant to stay private and be used for 
+	// bootstrapping another processes' module manager with a serialized state.
+	// It doesn't include any versioning as it is used with the
+	// the same binary executable for both the parent and 
+	// children processes. It also takes care or saving/restoring
+	// additional dll search directories.
+	TArray<FString> DllDirectories;
+	if (Ar.IsSaving())
+	{
+		TMap<FName, FString> OutModulePaths;
+		// Force any pending paths to be processed
+		FindModulePaths(TEXT("*"), OutModulePaths);
+
+		FPlatformProcess::GetDllDirectories(DllDirectories);
+	}
+
+	Ar << *this;
+	Ar << DllDirectories;
+
+	if (Ar.IsLoading())
+	{
+		for (const FString& DllDirectory : DllDirectories)
+		{
+			FPlatformProcess::AddDllDirectory(*DllDirectory);
+		}
+	}
+}
 #endif
 
 void FModuleManager::ResetModulePathsCache()
@@ -1084,8 +1162,7 @@ bool FModuleManager::LoadModuleWithCallback( const FName InModuleName, FOutputDe
 
 void FModuleManager::AddExtraBinarySearchPaths()
 {
-	static bool bAdded = false;
-	if (!bAdded)
+	if (!bExtraBinarySearchPathsAdded)
 	{
 		// Ensure that dependency dlls can be found in restricted sub directories
 		TArray<FString> RestrictedFolderNames = { TEXT("NoRedist"), TEXT("NotForLicensees"), TEXT("CarefullyRedist") };
@@ -1101,7 +1178,7 @@ void FModuleManager::AddExtraBinarySearchPaths()
 			}
 		}
 
-		bAdded = true;
+		bExtraBinarySearchPathsAdded = true;
 	}
 }
 
diff --git a/Engine/Source/Runtime/Core/Private/Windows/WindowsPlatformProcess.cpp b/Engine/Source/Runtime/Core/Private/Windows/WindowsPlatformProcess.cpp
--- a/Engine/Source/Runtime/Core/Private/Windows/WindowsPlatformProcess.cpp
+++ b/Engine/Source/Runtime/Core/Private/Windows/WindowsPlatformProcess.cpp
@@ -51,6 +51,11 @@ void FWindowsPlatformProcess::AddDllDirectory(const TCHAR* Directory)
 	DllDirectories.AddUnique(NormalizedDirectory);
 }
 
+void FWindowsPlatformProcess::GetDllDirectories(TArray<FString>& OutDllDirectories)
+{
+	OutDllDirectories = DllDirectories;
+}
+
 void* FWindowsPlatformProcess::GetDllHandle( const TCHAR* FileName )
 {
 	check(FileName);
diff --git a/Engine/Source/Runtime/Core/Public/GenericPlatform/GenericPlatformProcess.h b/Engine/Source/Runtime/Core/Public/GenericPlatform/GenericPlatformProcess.h
--- a/Engine/Source/Runtime/Core/Public/GenericPlatform/GenericPlatformProcess.h
+++ b/Engine/Source/Runtime/Core/Public/GenericPlatform/GenericPlatformProcess.h
@@ -174,6 +174,12 @@ struct CORE_API FGenericPlatformProcess
 
 	}
 
+	/** Get the list of registered directories to search in when resolving implicitly loaded or filename-only DLLs. **/
+	FORCEINLINE static void GetDllDirectories(TArray<FString>& OutDllDirectories)
+	{
+
+	}
+
 	/**
 	 * Retrieves the ProcessId of this process.
 	 *
diff --git a/Engine/Source/Runtime/Core/Public/Misc/ConfigCacheIni.h b/Engine/Source/Runtime/Core/Public/Misc/ConfigCacheIni.h
--- a/Engine/Source/Runtime/Core/Public/Misc/ConfigCacheIni.h
+++ b/Engine/Source/Runtime/Core/Public/Misc/ConfigCacheIni.h
@@ -249,8 +249,12 @@ public:
 
 	// look for "array of struct" keys for overwriting single entries of an array
 	TMap<FName, FString> ArrayOfStructKeys;
+
+	friend FArchive& operator<<(FArchive& Ar, FConfigSection& ConfigSection);
 };
 
+FArchive& operator<<(FArchive& Ar, FConfigSection& ConfigSection);
+
 /**
  * FIniFilename struct.
  * 
@@ -270,6 +274,16 @@ struct FIniFilename
 		, bRequired(InIsRequired) 
 		, CacheKey(InCacheKey)
 	{}
+
+	FIniFilename() = default;
+
+	friend FArchive& operator<<(FArchive& Ar, FIniFilename& IniFilename)
+	{
+		Ar << IniFilename.Filename;
+		Ar << IniFilename.bRequired;
+		Ar << IniFilename.CacheKey;
+		return Ar;
+	}
 };
 
 
@@ -290,6 +304,13 @@ private:
 public:
 	FConfigFileHierarchy();
 
+	friend FArchive& operator<<(FArchive& Ar, FConfigFileHierarchy& ConfigFileHierarchy)
+	{
+		Ar << static_cast<FConfigFileHierarchy::Super&>(ConfigFileHierarchy);
+		Ar << ConfigFileHierarchy.KeyGen;
+		return Ar;
+	}
+
 private:
 	int32 GenerateDynamicKey();
 
@@ -426,6 +447,7 @@ public:
 	/** Appends a new INI file to the SourceIniHierarchy and combines it */
 	CORE_API void AddDynamicLayerToHeirarchy(const FString& Filename);
 
+	friend FArchive& operator<<(FArchive& Ar, FConfigFile& ConfigFile);
 private:
 
 	// This holds per-object config class names, with their ArrayOfStructKeys. Since the POC sections are all unique,
@@ -465,6 +487,8 @@ private:
 	friend class FConfigCacheIni;
 };
 
+FArchive& operator<<(FArchive& Ar, FConfigFile& ConfigFile);
+
 /**
  * Declares a delegate type that's used by the config system to allow iteration of key value pairs.
  */
@@ -895,7 +919,16 @@ public:
 	 */
 	static FString GetGameUserSettingsDir();
 
+	/**
+	 * Save the current config cache state into a file for bootstrapping other processes.
+	 */
+	void SaveCurrentStateForBootstrap(const TCHAR* Filename);
+
+	friend FArchive& operator<<(FArchive& Ar, FConfigCacheIni& ConfigCacheIni);
 private:
+	/** Serialize a bootstrapping state into or from an archive */
+	void SerializeStateForBootstrap_Impl(FArchive& Ar);
+
 	/** true if file operations should not be performed */
 	bool bAreFileOperationsDisabled;
 
@@ -906,6 +939,8 @@ private:
 	EConfigCacheType Type;
 };
 
+FArchive& operator<<(FArchive& Ar, FConfigCacheIni& ConfigCacheIni);
+
 UE_DEPRECATED(4.24, "This functionality to generate Scalability@Level section string has been moved to Scalability.cpp. Explictly construct section you need manually.")
 CORE_API void ApplyCVarSettingsGroupFromIni(const TCHAR* InSectionBaseName, int32 InGroupNumber, const TCHAR* InIniFilename, uint32 SetBy);
 
diff --git a/Engine/Source/Runtime/Core/Public/Modules/ModuleManager.h b/Engine/Source/Runtime/Core/Public/Modules/ModuleManager.h
--- a/Engine/Source/Runtime/Core/Public/Modules/ModuleManager.h
+++ b/Engine/Source/Runtime/Core/Public/Modules/ModuleManager.h
@@ -423,6 +423,9 @@ public:
 
 	/** Determines if any non-default module instances are loaded (eg. hot reloaded modules) */
 	bool HasAnyOverridenModuleFilename() const;
+
+	/** Save the current module manager's state into a file for bootstrapping other processes. */
+	void SaveCurrentStateForBootstrap(const TCHAR* Filename);
 #endif
 
 	/**
@@ -544,6 +547,7 @@ public:
 	/** Clears module path cache */
 	void ResetModulePathsCache();
 
+	friend FArchive& operator<<( FArchive& Ar, FModuleManager& ModuleManager );
 private:
 	static void WarnIfItWasntSafeToLoadHere(const FName InModuleName);
 
@@ -567,6 +571,9 @@ private:
 
 	/** Finds modules within a given directory. */
 	void FindModulePathsInDirectory(const FString &DirectoryName, bool bIsGameDirectory, TMap<FName, FString> &OutModulePaths) const;
+
+	/** Serialize a bootstrapping state into or from an archive. */
+	void SerializeStateForBootstrap_Impl(FArchive& Ar);
 #endif
 
 	/** Adds pending module initializer registrations to the StaticallyLinkedModuleInitializers map. */
@@ -587,6 +594,9 @@ private:
 	/** True if module manager should automatically register new UObjects discovered while loading C++ modules */
 	bool bCanProcessNewlyLoadedObjects;
 
+	/** True once AddExtraBinarySearchPaths has been called */
+	bool bExtraBinarySearchPathsAdded;
+
 	/** Cache of known module paths. Used for performance. Can increase editor startup times by up to 30% */
 	mutable TMap<FName, FString> ModulePathsCache;
 
@@ -616,6 +626,8 @@ private:
 	mutable FCriticalSection ModulesCriticalSection;
 };
 
+FArchive& operator<<( FArchive& Ar, FModuleManager& ModuleManager );
+
 /**
  * Utility class for registering modules that are statically linked.
  */
diff --git a/Engine/Source/Runtime/Core/Public/Windows/WindowsPlatformProcess.h b/Engine/Source/Runtime/Core/Public/Windows/WindowsPlatformProcess.h
--- a/Engine/Source/Runtime/Core/Public/Windows/WindowsPlatformProcess.h
+++ b/Engine/Source/Runtime/Core/Public/Windows/WindowsPlatformProcess.h
@@ -129,6 +129,7 @@ public:
 	static void FreeDllHandle( void* DllHandle );
 	static void* GetDllExport( void* DllHandle, const TCHAR* ProcName );
 	static void AddDllDirectory(const TCHAR* Directory);
+	static void GetDllDirectories(TArray<FString>& OutDllDirectories);
 	static void PushDllDirectory(const TCHAR* Directory);
 	static void PopDllDirectory(const TCHAR* Directory);
 	static uint32 GetCurrentProcessId();
diff --git a/Engine/Source/Runtime/Launch/Private/LaunchEngineLoop.cpp b/Engine/Source/Runtime/Launch/Private/LaunchEngineLoop.cpp
--- a/Engine/Source/Runtime/Launch/Private/LaunchEngineLoop.cpp
+++ b/Engine/Source/Runtime/Launch/Private/LaunchEngineLoop.cpp
@@ -1405,6 +1405,8 @@ int32 FEngineLoop::PreInitPreStartupScreen(const TCHAR* CmdLine)
 	}
 
 	{
+		SCOPED_BOOT_TIMING("InitTrace")
+
 		FString Parameter;
 		if (FParse::Value(CmdLine, TEXT("-tracehost="), Parameter))
 		{
@@ -1596,7 +1598,10 @@ int32 FEngineLoop::PreInitPreStartupScreen(const TCHAR* CmdLine)
 
 	// 
 #if PLATFORM_DESKTOP && !IS_MONOLITHIC
-	FModuleManager::Get().AddExtraBinarySearchPaths();
+	{
+		SCOPED_BOOT_TIMING("AddExtraBinarySearchPaths");
+		FModuleManager::Get().AddExtraBinarySearchPaths();
+	}
 #endif
 
 	// Initialize file manager
