Remove 8 Above Normal IO threads in each shader workers by using -nothreading flag
  - This helps with overall system responsiveness when many shader workers are spawned at the same time

diff --git a/Engine/Source/Programs/ShaderCompileWorker/Private/ShaderCompileWorker.cpp b/Engine/Source/Programs/ShaderCompileWorker/Private/ShaderCompileWorker.cpp
--- a/Engine/Source/Programs/ShaderCompileWorker/Private/ShaderCompileWorker.cpp
+++ b/Engine/Source/Programs/ShaderCompileWorker/Private/ShaderCompileWorker.cpp
@@ -874,7 +874,7 @@ static void DirectCompile(const TArray<const class IShaderFormat*>& ShaderFormat
  */
 static int32 GuardedMain(int32 argc, TCHAR* argv[], bool bDirectMode)
 {
-	FString ExtraCmdLine = TEXT("-NOPACKAGECACHE -ReduceThreadUsage -cpuprofilertrace -nocrashreports");
+	FString ExtraCmdLine = TEXT("-NOPACKAGECACHE -ReduceThreadUsage -cpuprofilertrace -nocrashreports -nothreading");
 
 	// When executing tasks remotely through XGE, enumerating files requires tcp/ip round-trips with
 	// the initiator, which can slow down engine initialization quite drastically.
