Fixes crash in shader compile worker due to RDG_BUFFER_ACCESS that doesn't have a shader type

diff --git a/Engine/Source/Programs/ShaderCompileWorker/Private/ShaderCompileWorker.cpp b/Engine/Source/Programs/ShaderCompileWorker/Private/ShaderCompileWorker.cpp
--- a/Engine/Source/Programs/ShaderCompileWorker/Private/ShaderCompileWorker.cpp
+++ b/Engine/Source/Programs/ShaderCompileWorker/Private/ShaderCompileWorker.cpp
@@ -483,6 +483,11 @@ private:
 					InputFile << NumElements;
 					InputFile << StructMetadataIndex;
 
+					if (ShaderType == nullptr)
+					{
+						ShaderType = TEXT("");
+					}
+
 					const FShaderParametersMetadata* StructMetadata = nullptr;
 					if (StructMetadataIndex != INDEX_NONE)
 					{
diff --git a/Engine/Source/Runtime/RenderCore/Public/ShaderParameterMetadata.h b/Engine/Source/Runtime/RenderCore/Public/ShaderParameterMetadata.h
--- a/Engine/Source/Runtime/RenderCore/Public/ShaderParameterMetadata.h
+++ b/Engine/Source/Runtime/RenderCore/Public/ShaderParameterMetadata.h
@@ -154,7 +154,9 @@ public:
 		,	NumColumns(InNumColumns)
 		,	NumElements(InNumElements)
 		,	Struct(InStruct)
-		{}
+		{
+			check(InShaderType);
+		}
 
 		/** Returns the string of the name of the element or name of the array of elements. */
 		const TCHAR* GetName() const { return Name; }
