Replace "UE4" with UnrealEditor in SCW error message.

#rnx

diff --git a/Engine/Source/Programs/ShaderCompileWorker/Private/ShaderCompileWorker.cpp b/Engine/Source/Programs/ShaderCompileWorker/Private/ShaderCompileWorker.cpp
--- a/Engine/Source/Programs/ShaderCompileWorker/Private/ShaderCompileWorker.cpp
+++ b/Engine/Source/Programs/ShaderCompileWorker/Private/ShaderCompileWorker.cpp
@@ -1245,7 +1245,7 @@ INT32_MAIN_INT32_ARGC_TCHAR_ARGV()
 	{
 		if (ArgC < 6)
 		{
-			printf("ShaderCompileWorker is called by UE4, it requires specific command like arguments.\n");
+			printf("ShaderCompileWorker is called by UnrealEditor, it requires specific command like arguments.\n");
 			return -1;
 		}
 
