Remove crash reporter thread from shader compiler as -nocrashreports was used anyway

#rnx

diff --git a/Engine/Source/Programs/ShaderCompileWorker/ShaderCompileWorker.Target.cs b/Engine/Source/Programs/ShaderCompileWorker/ShaderCompileWorker.Target.cs
--- a/Engine/Source/Programs/ShaderCompileWorker/ShaderCompileWorker.Target.cs
+++ b/Engine/Source/Programs/ShaderCompileWorker/ShaderCompileWorker.Target.cs
@@ -58,5 +58,8 @@ public class ShaderCompileWorkerTarget : TargetRules
 
 		// Disable external profiling in ShaderCompiler to improve startup time
 		GlobalDefinitions.Add("UE_EXTERNAL_PROFILING_ENABLED=0");
+
+		// This removes another thread being created even when -nocrashreports is specified
+		GlobalDefinitions.Add("NOINITCRASHREPORTER=1");
 	}
 }
