Fix parsing multi-entrypoint arguments for parameter "-entry" in SCW.

#rnx

diff --git a/Engine/Source/Programs/ShaderCompileWorker/Private/ShaderCompileWorker.cpp b/Engine/Source/Programs/ShaderCompileWorker/Private/ShaderCompileWorker.cpp
--- a/Engine/Source/Programs/ShaderCompileWorker/Private/ShaderCompileWorker.cpp
+++ b/Engine/Source/Programs/ShaderCompileWorker/Private/ShaderCompileWorker.cpp
@@ -933,6 +933,12 @@ static void DirectCompile(const TArray<const class IShaderFormat*>& ShaderFormat
 			else if (Token.StartsWith(TEXT("entry=")))
 			{
 				Entry = Token.RightChop(6);
+
+				// Remove quotations marks at beginning and end; happens when multiple entry points are specified, e.g. -entry="closesthit=A anyhit=B"
+				if (Entry.Len() >= 2 && Entry[0] == TEXT('\"') && Entry[Entry.Len() - 1] == TEXT('\"'))
+				{
+					Entry = Entry.Mid(1, Entry.Len() - 2);
+				}
 			}
 			else if (Token.StartsWith(TEXT("cflags=")))
 			{
