Temp-fix for non-Windows platforms

diff --git a/Engine/Shaders/Public/Platform/Metal/MetalCommon.ush b/Engine/Shaders/Public/Platform/Metal/MetalCommon.ush
--- a/Engine/Shaders/Public/Platform/Metal/MetalCommon.ush
+++ b/Engine/Shaders/Public/Platform/Metal/MetalCommon.ush
@@ -5,7 +5,7 @@
 =============================================================================*/
 
 // Update this GUID to improve shader recompilation for Metal only shaders
-// GUID = 1F6F3707-72C3-4017-8111-CD48182B0579
+// GUID = 35B87765-455B-4833-8A78-0804AD8818E7
 
 #pragma once
 
diff --git a/Engine/Shaders/Public/Platform/Vulkan/VulkanCommon.ush b/Engine/Shaders/Public/Platform/Vulkan/VulkanCommon.ush
--- a/Engine/Shaders/Public/Platform/Vulkan/VulkanCommon.ush
+++ b/Engine/Shaders/Public/Platform/Vulkan/VulkanCommon.ush
@@ -5,7 +5,7 @@
 =============================================================================*/
 
 // Update this GUID to improve shader recompilation for Vulkan only shaders
-// GUID = 1C88AA5C-AFC4-49BB-A5C8-88667ADA03B1
+// GUID = 3C913DA7-4D55-4B21-9D28-5FDF44958081
 
 #pragma once
 
diff --git a/Engine/Source/Programs/ShaderCompileWorker/Private/ShaderCompileWorker.cpp b/Engine/Source/Programs/ShaderCompileWorker/Private/ShaderCompileWorker.cpp
--- a/Engine/Source/Programs/ShaderCompileWorker/Private/ShaderCompileWorker.cpp
+++ b/Engine/Source/Programs/ShaderCompileWorker/Private/ShaderCompileWorker.cpp
@@ -601,7 +601,9 @@ private:
 			// (Otherwise quick changes to a shader file can result in the wrong output.)
 			FName ShaderPlatformName;
 			InputFile << ShaderPlatformName;
-			FlushShaderFileCache(&ShaderPlatformName);
+
+			// temp-fix for non-Windows platforms
+			//FlushShaderFileCache(&ShaderPlatformName);
 			
 			for (int32 BatchIndex = 0; BatchIndex < NumBatches; BatchIndex++)
 			{
