Added the ability to do Unreal Insights traces of the Shader Compiler Worker process.

This can be enabled by modifying `UnrealEngine\Engine\Saved\UnrealBuildTool\BuildConfiguration.xml` like so:

    <?xml version="1.0" encoding="utf-8" ?>
    <Configuration xmlns="https://www.unrealengine.com/BuildConfiguration">
        <BuildConfiguration>
            <bShaderCompilerWorkerTrace>true</bShaderCompilerWorkerTrace>
        </BuildConfiguration>
    </Configuration>

- Added a build configuration xml value, `bShaderCompilerWorkerTrace`.
- Turning this on will set USE_SHADER_COMPILER_WORKER_TRACE=1
- Move the parameter -nothreading to be set when we launch the process instead of internally as an extra cmd line arg.
- Unreal Insights uses a separate thread to send events so threading support is needed for the program.  When we have USE_SHADER_COMPILER_WORKER_TRACE enabled we need to turn off `-nothreading`.
- When USE_SHADER_COMPILER_WORKER_TRACE is enabled we pass in `-trace=default` to get CPU event markers.
- The SCW program needs to turn on the following defines to be able to perform CPU and memory traces:
	ENABLE_LOW_LEVEL_MEM_TRACKER=1
	UE_MEMORY_TAGS_TRACE_ENABLED=1
	UE_TRACE_ENABLED=1
- Instrument Shader Compiler Worker with TRACE_CPUPROFILER_EVENT_SCOPE.  This are no-ops when this is turned off.

todo: Make the shader compiler worker inherit the trace args from the main process it was launched from.

diff --git a/Engine/Source/Programs/ShaderCompileWorker/Private/ShaderCompileWorker.cpp b/Engine/Source/Programs/ShaderCompileWorker/Private/ShaderCompileWorker.cpp
--- a/Engine/Source/Programs/ShaderCompileWorker/Private/ShaderCompileWorker.cpp
+++ b/Engine/Source/Programs/ShaderCompileWorker/Private/ShaderCompileWorker.cpp
@@ -110,6 +110,8 @@ static const IShaderFormat* FindShaderFormat(FName Name)
 /** Processes a compilation job. */
 static void ProcessCompilationJob(const FShaderCompilerInput& Input,FShaderCompilerOutput& Output,const FString& WorkingDirectory)
 {
+	TRACE_CPUPROFILER_EVENT_SCOPE(ProcessCompilationJob);
+
 	const IShaderFormat* Compiler = FindShaderFormat(Input.ShaderFormat);
 	if (!Compiler)
 	{
@@ -158,6 +160,8 @@ static void UpdateFileSize(FArchive& OutputFile, int64 FileSizePosition)
 static int64 WriteOutputFileHeader(FArchive& OutputFile, int32 ErrorCode, int32 CallstackLength, const TCHAR* Callstack,
 	int32 ExceptionInfoLength, const TCHAR* ExceptionInfo)
 {
+	TRACE_CPUPROFILER_EVENT_SCOPE(WriteOutputFileHeader);
+
 	int64 FileSizePosition = 0;
 	int32 OutputVersion = ShaderCompileWorkerOutputVersion;
 	OutputFile << OutputVersion;
@@ -225,6 +229,7 @@ public:
 	void Loop()
 	{
 		UE_LOG(LogShaders, Log, TEXT("Entering job loop"));
+		TRACE_CPUPROFILER_EVENT_SCOPE(Loop);
 
 		while(true)
 		{
@@ -233,6 +238,8 @@ public:
 
 			// Read & Process Input
 			{
+				TRACE_CPUPROFILER_EVENT_SCOPE(ReadInput);
+
 				FArchive* InputFilePtr = OpenInputFile();
 				if(!InputFilePtr)
 				{
@@ -282,6 +289,14 @@ public:
 				UE_LOG(LogShaders, Log, TEXT("TimeToLive set to 0, or used HLSLcc compiler, exiting after single job"));
 				break;
 			}
+
+#if ENABLE_LOW_LEVEL_MEM_TRACKER
+			{
+				TRACE_CPUPROFILER_EVENT_SCOPE(UpdateStatsPerFrame);
+
+				FLowLevelMemTracker::Get().UpdateStatsPerFrame();
+			}
+#endif
 		}
 
 		UE_LOG(LogShaders, Log, TEXT("Exiting job loop"));
@@ -347,6 +362,8 @@ private:
 
 	void ProcessInputFromArchive(FArchive* InputFilePtr, TArray<FJobResult>& OutSingleJobResults, TArray<FPipelineJobResult>& OutPipelineJobResults)
 	{
+		TRACE_CPUPROFILER_EVENT_SCOPE(ProcessInputFromArchive);
+
 		int32 InputVersion;
 		*InputFilePtr << InputVersion;
 		if (ShaderCompileWorkerInputVersion != InputVersion)
@@ -734,6 +751,8 @@ private:
 
 	FArchive* CreateOutputArchive()
 	{
+		TRACE_CPUPROFILER_EVENT_SCOPE(CreateOutputArchive);
+
 		FArchive* OutputFilePtr = nullptr;
 		const double StartTime = FPlatformTime::Seconds();
 		bool bResult = false;
@@ -785,6 +804,8 @@ private:
 
 	void WriteToOutputArchive(FArchive* OutputFilePtr, TArray<FJobResult>& SingleJobResults, TArray<FPipelineJobResult>& PipelineJobResults)
 	{
+		TRACE_CPUPROFILER_EVENT_SCOPE(WriteToOutputArchive);
+
 		FArchive& OutputFile = *OutputFilePtr;
 		int64 FileSizePosition = WriteOutputFileHeader(OutputFile, (int32)ESCWErrorCode::Success, 0, nullptr, 0, nullptr);
 
@@ -900,6 +921,8 @@ private:
 	
 	static bool AnyJobUsedHLSLccCompiler(TArray<FJobResult>& SingleJobResults, TArray<FPipelineJobResult>& PipelineJobResults)
 	{
+		TRACE_CPUPROFILER_EVENT_SCOPE(AnyJobUsedHLSLccCompiler);
+
 		for (int32 ResultIndex = 0; ResultIndex < SingleJobResults.Num(); ResultIndex++)
 		{
 			FJobResult& JobResult = SingleJobResults[ResultIndex];
@@ -927,6 +950,8 @@ private:
 
 static void DirectCompile(const TArray<const class IShaderFormat*>& ShaderFormats)
 {
+	TRACE_CPUPROFILER_EVENT_SCOPE(DirectCompile);
+
 	// Find all the info required for compiling a single shader
 	TArray<FString> Tokens, Switches;
 	FCommandLine::Parse(FCommandLine::Get(), Tokens, Switches);
@@ -1083,7 +1108,7 @@ static void DirectCompile(const TArray<const class IShaderFormat*>& ShaderFormat
  */
 static int32 GuardedMain(int32 argc, TCHAR* argv[], bool bDirectMode)
 {
-	FString ExtraCmdLine = TEXT("-NOPACKAGECACHE -ReduceThreadUsage -cpuprofilertrace -nocrashreports -nothreading");
+	FString ExtraCmdLine = TEXT("-NOPACKAGECACHE -ReduceThreadUsage -cpuprofilertrace -nocrashreports");
 
 	// When executing tasks remotely through XGE, enumerating files requires tcp/ip round-trips with
 	// the initiator, which can slow down engine initialization quite drastically.
@@ -1150,6 +1175,8 @@ static int32 GuardedMain(int32 argc, TCHAR* argv[], bool bDirectMode)
 	GLogConsole->Show( true );
 #endif
 
+	TRACE_CPUPROFILER_EVENT_SCOPE(Main);
+
 	auto AtomicSave = 
 		[](const FString& Filename, TFunctionRef<void (const FString& TmpFile)> SaveFunction)
 		{
@@ -1177,6 +1204,8 @@ static int32 GuardedMain(int32 argc, TCHAR* argv[], bool bDirectMode)
 	TMap<FString, uint32> FormatVersionMap;
 	for (int32 Index = 0; Index < ShaderFormats.Num(); Index++)
 	{
+		TRACE_CPUPROFILER_EVENT_SCOPE(ShaderFormat);
+
 		TArray<FName> OutFormats;
 		ShaderFormats[Index]->GetSupportedFormats(OutFormats);
 		check(OutFormats.Num());
@@ -1201,6 +1230,8 @@ static int32 GuardedMain(int32 argc, TCHAR* argv[], bool bDirectMode)
 		SetConsoleTitle(argv[3]);
 #endif
 
+		TRACE_CPUPROFILER_EVENT_SCOPE(FWorkLoop);
+
 		FWorkLoop WorkLoop(argv[2], argv[1], argv[4], argv[5], FormatVersionMap);
 		WorkLoop.Loop();
 	}
diff --git a/Engine/Source/Programs/ShaderCompileWorker/ShaderCompileWorker.Target.cs b/Engine/Source/Programs/ShaderCompileWorker/ShaderCompileWorker.Target.cs
--- a/Engine/Source/Programs/ShaderCompileWorker/ShaderCompileWorker.Target.cs
+++ b/Engine/Source/Programs/ShaderCompileWorker/ShaderCompileWorker.Target.cs
@@ -62,5 +62,12 @@ public class ShaderCompileWorkerTarget : TargetRules
 
 		// This removes another thread being created even when -nocrashreports is specified
 		GlobalDefinitions.Add("NOINITCRASHREPORTER=1");
+
+		if (bShaderCompilerWorkerTrace)
+        {
+			GlobalDefinitions.Add("ENABLE_LOW_LEVEL_MEM_TRACKER=1");
+			GlobalDefinitions.Add("UE_MEMORY_TAGS_TRACE_ENABLED=1");
+			GlobalDefinitions.Add("UE_TRACE_ENABLED=1");
+		}
 	}
 }
diff --git a/Engine/Source/Programs/UnrealBuildTool/Configuration/TargetRules.cs b/Engine/Source/Programs/UnrealBuildTool/Configuration/TargetRules.cs
--- a/Engine/Source/Programs/UnrealBuildTool/Configuration/TargetRules.cs
+++ b/Engine/Source/Programs/UnrealBuildTool/Configuration/TargetRules.cs
@@ -1651,6 +1651,12 @@ namespace UnrealBuildTool
 		[XmlConfigFile(Category = "BuildConfiguration")]
 		public bool bUseMallocProfiler = false;
 
+		/// <summary>
+		/// If true, then enable Unreal Insights (utrace) profiling in the build for the Shader Compiler Worker (defines USE_SHADER_COMPILER_WORKER_TRACE=1).
+		/// </summary>
+		[XmlConfigFile(Category = "BuildConfiguration")]
+		public bool bShaderCompilerWorkerTrace = false;
+
 		/// <summary>
 		/// Enables "Shared PCHs", a feature which significantly speeds up compile times by attempting to
 		/// share certain PCH files between modules that UBT detects is including those PCH's header files.
@@ -3296,6 +3302,11 @@ namespace UnrealBuildTool
 			get { return Inner.bUseMallocProfiler; }
 		}
 
+		public bool bShaderCompilerWorkerTrace
+        {
+			get { return Inner.bShaderCompilerWorkerTrace;  }
+        }
+
 		public bool bUseSharedPCHs
 		{
 			get { return Inner.bUseSharedPCHs; }
diff --git a/Engine/Source/Programs/UnrealBuildTool/System/RulesAssembly.cs b/Engine/Source/Programs/UnrealBuildTool/System/RulesAssembly.cs
--- a/Engine/Source/Programs/UnrealBuildTool/System/RulesAssembly.cs
+++ b/Engine/Source/Programs/UnrealBuildTool/System/RulesAssembly.cs
@@ -630,6 +630,16 @@ namespace UnrealBuildTool
 				Rules.GlobalDefinitions.Add("USE_MALLOC_PROFILER=1");
 			}
 
+			// Setup utrace for Shader Compiler Worker
+			if (Rules.bShaderCompilerWorkerTrace)
+            {
+				Rules.GlobalDefinitions.Add("USE_SHADER_COMPILER_WORKER_TRACE=1");
+			}
+            else
+            {
+				Rules.GlobalDefinitions.Add("USE_SHADER_COMPILER_WORKER_TRACE=0");
+            }
+
 			// Set a macro if we allow using generated inis
 			if (!Rules.bAllowGeneratedIniWhenCooked)
 			{
diff --git a/Engine/Source/Runtime/Engine/Private/ShaderCompiler/ShaderCompiler.cpp b/Engine/Source/Runtime/Engine/Private/ShaderCompiler/ShaderCompiler.cpp
--- a/Engine/Source/Runtime/Engine/Private/ShaderCompiler/ShaderCompiler.cpp
+++ b/Engine/Source/Runtime/Engine/Private/ShaderCompiler/ShaderCompiler.cpp
@@ -4166,6 +4166,13 @@ FProcHandle FShaderCompilingManager::LaunchWorker(const FString& WorkingDirector
 	}
 	WorkerParameters += FCommandLine::GetSubprocessCommandline();
 
+#if USE_SHADER_COMPILER_WORKER_TRACE
+	// When doing utrace functionality we can't run with -nothreading, since it won't create the utrace thread to send events.
+	WorkerParameters += FString(TEXT(" -trace=default "));
+#else
+	WorkerParameters += FString(TEXT(" -nothreading "));
+#endif // USE_SHADER_COMPILER_WORKER_TRACE
+
 	// Launch the worker process
 	int32 PriorityModifier = -1; // below normal
 	GConfig->GetInt(TEXT("DevOptions.Shaders"), TEXT("WorkerProcessPriority"), PriorityModifier, GEngineIni);
