[shader compilation] changes to error handling in ShaderCompileWorker
- each platform compilation request is now wrapped in __try/__except rather than a single exception handling block at the top level of the worker; this allows us to log an exception (with callstack) as a compilation error and continue the batch
- remove SEH code from ShaderConductorContext, the above makes this redundant (and it didn't provide any actionable information)
- strip down SEH code in D3DShaderCompiler; now only used for the purposes of pre-compiling with DXC in the case of an FXC crash. dumping preprocessed source will be handled in a different manner in a forthcoming CL
- minor change in the DXC precompile path to not log an unnecessary warning when performing an explicitly-requested DXC precompile

diff --git a/Engine/Source/Developer/ShaderCompilerCommon/Private/ShaderConductorContext.cpp b/Engine/Source/Developer/ShaderCompilerCommon/Private/ShaderConductorContext.cpp
--- a/Engine/Source/Developer/ShaderCompilerCommon/Private/ShaderConductorContext.cpp
+++ b/Engine/Source/Developer/ShaderCompilerCommon/Private/ShaderConductorContext.cpp
@@ -17,8 +17,7 @@ namespace CrossCompiler
 
 #if PLATFORM_MAC || PLATFORM_WINDOWS || PLATFORM_LINUX
 
-	// Inner wrapper function is required here because '__try'-statement cannot be used with function that requires object unwinding
-	static void InnerScRewriteWrapper(
+	static void ScRewriteWrapper(
 		const ShaderConductor::Compiler::SourceDesc& InSourceDesc,
 		const ShaderConductor::Compiler::Options& InOptions,
 		ShaderConductor::Compiler::ResultDesc& OutResultDesc)
@@ -26,33 +25,7 @@ namespace CrossCompiler
 		OutResultDesc = ShaderConductor::Compiler::Rewrite(InSourceDesc, InOptions);
 	}
 
-	static bool ScRewriteWrapper(
-		const ShaderConductor::Compiler::SourceDesc& InSourceDesc,
-		const ShaderConductor::Compiler::Options& InOptions,
-		ShaderConductor::Compiler::ResultDesc& OutResultDesc,
-		bool& bOutException)
-	{
-		bOutException = false;
-#if !PLATFORM_SEH_EXCEPTIONS_DISABLED
-		__try
-#endif
-		{
-			InnerScRewriteWrapper(InSourceDesc, InOptions, OutResultDesc);
-			return true;
-		}
-#if !PLATFORM_SEH_EXCEPTIONS_DISABLED
-		__except (EXCEPTION_EXECUTE_HANDLER)
-		{
-			GSCWErrorCode = ESCWErrorCode::CrashInsidePlatformCompiler;
-			FMemory::Memzero(OutResultDesc);
-			bOutException = true;
-			return false;
-		}
-#endif
-	}
-
-	// Inner wrapper function is required here because '__try'-statement cannot be used with function that requires object unwinding
-	static void InnerScCompileWrapper(
+	static void ScCompileWrapper(
 		const ShaderConductor::Compiler::SourceDesc& InSourceDesc,
 		const ShaderConductor::Compiler::Options& InOptions,
 		const ShaderConductor::Compiler::TargetDesc& InTargetDesc,
@@ -61,34 +34,7 @@ namespace CrossCompiler
 		OutResultDesc = ShaderConductor::Compiler::Compile(InSourceDesc, InOptions, InTargetDesc);
 	}
 
-	static bool ScCompileWrapper(
-		const ShaderConductor::Compiler::SourceDesc& InSourceDesc,
-		const ShaderConductor::Compiler::Options& InOptions,
-		const ShaderConductor::Compiler::TargetDesc& InTargetDesc,
-		ShaderConductor::Compiler::ResultDesc& OutResultDesc,
-		bool& bOutException)
-	{
-		bOutException = false;
-#if !PLATFORM_SEH_EXCEPTIONS_DISABLED
-		__try
-#endif
-		{
-			InnerScCompileWrapper(InSourceDesc, InOptions, InTargetDesc, OutResultDesc);
-			return true;
-		}
-#if !PLATFORM_SEH_EXCEPTIONS_DISABLED
-		__except (EXCEPTION_EXECUTE_HANDLER)
-		{
-			GSCWErrorCode = ESCWErrorCode::CrashInsidePlatformCompiler;
-			FMemory::Memzero(OutResultDesc);
-			bOutException = true;
-			return false;
-		}
-#endif
-	}
-
-	// Inner wrapper function is required here because '__try'-statement cannot be used with function that requires object unwinding
-	static void InnerScCompileWrapper(
+	static void ScConvertBinaryWrapper(
 		const ShaderConductor::Compiler::ResultDesc& InBinaryDesc,
 		const ShaderConductor::Compiler::SourceDesc& InSourceDesc,
 		const ShaderConductor::Compiler::TargetDesc& InTargetDesc,
@@ -98,33 +44,6 @@ namespace CrossCompiler
 		OutResultDesc = ShaderConductor::Compiler::ConvertBinary(InBinaryDesc, InSourceDesc, InOptions, InTargetDesc);
 	}
 
-	static bool ScConvertBinaryWrapper(
-		const ShaderConductor::Compiler::ResultDesc& InBinaryDesc,
-		const ShaderConductor::Compiler::SourceDesc& InSourceDesc,
-		const ShaderConductor::Compiler::TargetDesc& InTargetDesc,
-		const ShaderConductor::Compiler::Options& InOptions,
-		ShaderConductor::Compiler::ResultDesc& OutResultDesc,
-		bool& bOutException)
-	{
-		bOutException = false;
-#if !PLATFORM_SEH_EXCEPTIONS_DISABLED
-		__try
-#endif
-		{
-			InnerScCompileWrapper(InBinaryDesc, InSourceDesc, InTargetDesc, InOptions, OutResultDesc);
-			return true;
-		}
-#if !PLATFORM_SEH_EXCEPTIONS_DISABLED
-		__except (EXCEPTION_EXECUTE_HANDLER)
-		{
-			GSCWErrorCode = ESCWErrorCode::CrashInsidePlatformCompiler;
-			FMemory::Memzero(OutResultDesc);
-			bOutException = true;
-			return false;
-		}
-#endif
-	}
-
 	// Converts the byte array 'InString' (without null terminator) to the output ANSI string 'OutString' (with appended null terminator).
 	static void ConvertByteArrayToAnsiString(const ANSICHAR* InString, uint32 InStringLength, TArray<ANSICHAR>& OutString)
 	{
@@ -616,11 +535,10 @@ namespace CrossCompiler
 
 		// Rewrite HLSL with wrapper function to catch exceptions from ShaderConductor
 		bool bSucceeded = false;
-		bool bException = false;
 		ShaderConductor::Compiler::ResultDesc ResultDesc;
-		ScRewriteWrapper(ScSourceDesc, ScOptions, ResultDesc, bException);
+		ScRewriteWrapper(ScSourceDesc, ScOptions, ResultDesc);
 
-		if (!ResultDesc.hasError && !bException && ResultDesc.target.Size() > 0)
+		if (!ResultDesc.hasError && ResultDesc.target.Size() > 0)
 		{
 			// Copy rewritten HLSL code into intermediate source code.
 			ConvertByteArrayToAnsiString(reinterpret_cast<const ANSICHAR*>(ResultDesc.target.Data()), ResultDesc.target.Size(), Intermediates->ShaderSource);
@@ -632,14 +550,6 @@ namespace CrossCompiler
 			}
 			bSucceeded = true;
 		}
-		else
-		{
-			if (bException)
-			{
-				Errors.Add(TEXT("ShaderConductor exception during rewrite"));
-			}
-			bSucceeded = false;
-		}
 
 		// Append compile error and warning to output reports
 		ConvertScCompileErrors(ResultDesc.errorWarningMsg, Errors);
@@ -662,25 +572,16 @@ namespace CrossCompiler
 
 		// Compile HLSL source code to SPIR-V
 		bool bSucceeded = false;
-		bool bException = false;
 		ShaderConductor::Compiler::ResultDesc ResultDesc;
-		ScCompileWrapper(ScSourceDesc, ScOptions, ScTargetDesc, ResultDesc, bException);
+		ScCompileWrapper(ScSourceDesc, ScOptions, ScTargetDesc, ResultDesc);
 
-		if (!ResultDesc.hasError && !bException && ResultDesc.target.Size() > 0)
+		if (!ResultDesc.hasError && ResultDesc.target.Size() > 0)
 		{
 			// Copy result blob into output SPIR-V module
 			OutSpirv = TArray<uint32>(reinterpret_cast<const uint32*>(ResultDesc.target.Data()), ResultDesc.target.Size() / 4);
 			bSucceeded = true;
 		}
-		else
-		{
-			if (bException)
-			{
-				Errors.Add(TEXT("ShaderConductor exception during compilation"));
-			}
-			bSucceeded = false;
-		}
-
+		
 		// Append compile error and warning to output reports
 		ConvertScCompileErrors(ResultDesc.errorWarningMsg, Errors);
 
@@ -772,25 +673,16 @@ namespace CrossCompiler
 
 		// Convert the input SPIR-V into Metal high level source
 		bool bSucceeded = false;
-		bool bException = false;
 		ShaderConductor::Compiler::ResultDesc ResultDesc;
-		ScConvertBinaryWrapper(ScBinaryDesc, ScSourceDesc, ScTargetDesc, ScOptions, ResultDesc, bException);
+		ScConvertBinaryWrapper(ScBinaryDesc, ScSourceDesc, ScTargetDesc, ScOptions, ResultDesc);
 
-		if (!ResultDesc.hasError && !bException && ResultDesc.target.Size() > 0)
+		if (!ResultDesc.hasError && ResultDesc.target.Size() > 0)
 		{
 			// Copy result blob into output SPIR-V module
 			OutputCallback(ResultDesc.target.Data(), ResultDesc.target.Size());
 			bSucceeded = true;
 		}
-		else
-		{
-			if (bException)
-			{
-				Errors.Add(TEXT("ShaderConductor exception during SPIR-V binary conversion"));
-			}
-			bSucceeded = false;
-		}
-
+	
 		// Append compile error and warning to output reports
 		if (ResultDesc.errorWarningMsg.Size() > 0)
 		{
diff --git a/Engine/Source/Developer/Windows/ShaderFormatD3D/Private/D3DShaderCompiler.cpp b/Engine/Source/Developer/Windows/ShaderFormatD3D/Private/D3DShaderCompiler.cpp
--- a/Engine/Source/Developer/Windows/ShaderFormatD3D/Private/D3DShaderCompiler.cpp
+++ b/Engine/Source/Developer/Windows/ShaderFormatD3D/Private/D3DShaderCompiler.cpp
@@ -442,54 +442,9 @@ static bool GetD3DCompilerFuncs(const FString& NewCompilerPath, pD3DCompile* Out
 	return false;
 }
 
-static const char* Win32SehExceptionToString(DWORD Code)
+static int D3DExceptionFilter(bool bCatchException)
 {
-#define CASE_TO_STRING(IDENT) case IDENT: return #IDENT
-
-	switch (Code)
-	{
-		CASE_TO_STRING(EXCEPTION_ACCESS_VIOLATION);
-		CASE_TO_STRING(EXCEPTION_ARRAY_BOUNDS_EXCEEDED);
-		CASE_TO_STRING(EXCEPTION_BREAKPOINT);
-		CASE_TO_STRING(EXCEPTION_DATATYPE_MISALIGNMENT);
-		CASE_TO_STRING(EXCEPTION_FLT_DENORMAL_OPERAND);
-		CASE_TO_STRING(EXCEPTION_FLT_DIVIDE_BY_ZERO);
-		CASE_TO_STRING(EXCEPTION_FLT_INEXACT_RESULT);
-		CASE_TO_STRING(EXCEPTION_FLT_INVALID_OPERATION);
-		CASE_TO_STRING(EXCEPTION_FLT_OVERFLOW);
-		CASE_TO_STRING(EXCEPTION_FLT_STACK_CHECK);
-		CASE_TO_STRING(EXCEPTION_FLT_UNDERFLOW);
-		CASE_TO_STRING(EXCEPTION_GUARD_PAGE);
-		CASE_TO_STRING(EXCEPTION_ILLEGAL_INSTRUCTION);
-		CASE_TO_STRING(EXCEPTION_IN_PAGE_ERROR);
-		CASE_TO_STRING(EXCEPTION_INT_DIVIDE_BY_ZERO);
-		CASE_TO_STRING(EXCEPTION_INT_OVERFLOW);
-		CASE_TO_STRING(EXCEPTION_INVALID_DISPOSITION);
-		CASE_TO_STRING(EXCEPTION_INVALID_HANDLE);
-		CASE_TO_STRING(EXCEPTION_NONCONTINUABLE_EXCEPTION);
-		CASE_TO_STRING(EXCEPTION_PRIV_INSTRUCTION);
-		CASE_TO_STRING(EXCEPTION_SINGLE_STEP);
-		CASE_TO_STRING(EXCEPTION_STACK_OVERFLOW);
-		CASE_TO_STRING(STATUS_UNWIND_CONSOLIDATE);
-		default: return nullptr;
-	}
-
-#undef CASE_TO_STRING
-}
-
-struct FD3DExceptionInfo
-{
-	uint32 Code;
-	uint64 Base;
-	uint64 Address;
-};
-
-static int D3DExceptionFilter(DWORD Code, LPEXCEPTION_POINTERS InInfo, FD3DExceptionInfo& OutInfo)
-{
-	OutInfo.Code = InInfo->ExceptionRecord->ExceptionCode;
-	OutInfo.Base = (uint64)GetModuleHandle(NULL);
-	OutInfo.Address = (uint64)InInfo->ExceptionRecord->ExceptionAddress;
-	return EXCEPTION_EXECUTE_HANDLER;
+	return bCatchException ? EXCEPTION_EXECUTE_HANDLER : EXCEPTION_CONTINUE_SEARCH;
 }
 
 static HRESULT D3DCompileWrapper(
@@ -505,8 +460,7 @@ static HRESULT D3DCompileWrapper(
 	uint32					Flags2,
 	ID3DBlob**				ppCode,
 	ID3DBlob**				ppErrorMsgs,
-	bool&					bOutException,
-	FD3DExceptionInfo&		OutExceptionInfo
+	bool					bCatchException = false
 	)
 {
 #if !PLATFORM_SEH_EXCEPTIONS_DISABLED
@@ -528,10 +482,9 @@ static HRESULT D3DCompileWrapper(
 		);
 	}
 #if !PLATFORM_SEH_EXCEPTIONS_DISABLED
-	__except(D3DExceptionFilter(GetExceptionCode(), GetExceptionInformation(), OutExceptionInfo))
+	__except(D3DExceptionFilter(bCatchException))
 	{
 		GSCWErrorCode = ESCWErrorCode::CrashInsidePlatformCompiler;
-		bOutException = true;
 		return E_FAIL;
 	}
 #endif
@@ -706,9 +659,6 @@ bool CompileAndProcessD3DShaderFXC(FString& PreprocessedShaderSource, const FStr
 
 	if (D3DCompileFunc)
 	{
-		bool bException = false;
-		FD3DExceptionInfo ExceptionInfo{};
-
 		const bool bHlslVersion2021 = Input.Environment.CompilerFlags.Contains(CFLAG_HLSL2021);
 		const bool bPrecompileWithDXC = bHlslVersion2021 || Input.Environment.CompilerFlags.Contains(CFLAG_PrecompileWithDXC);
 		if (!bPrecompileWithDXC)
@@ -726,20 +676,16 @@ bool CompileAndProcessD3DShaderFXC(FString& PreprocessedShaderSource, const FStr
 				0,
 				Shader.GetInitReference(),
 				Errors.GetInitReference(),
-				bException,
-				ExceptionInfo
+				// We only want to catch the exception on initial FXC compiles so we can retry with a 
+				// DXC precompilation step. If it fails again on the second attempt then we let
+				// ShaderCompileWorker handle the exception and log an error.
+				/* bCatchException */ true
 			);
 		}
 
 		// Some materials give FXC a hard time to optimize and the compiler fails with an internal error.
-		if (bPrecompileWithDXC || Result == HRESULT_FROM_WIN32(ERROR_ARITHMETIC_OVERFLOW) || Result == E_OUTOFMEMORY || bException || (Result != S_OK && CompileErrorsContainInternalError(Errors.GetReference())))
+		if (bPrecompileWithDXC || Result == HRESULT_FROM_WIN32(ERROR_ARITHMETIC_OVERFLOW) || Result == E_OUTOFMEMORY || Result == E_FAIL || (Result != S_OK && CompileErrorsContainInternalError(Errors.GetReference())))
 		{
-			if (bPrecompileWithDXC)
-			{
-				// Let the user know this shader had to be cross-compiled due to a crash in FXC. Only shows up if CVar 'r.ShaderDevelopmentMode' is enabled.
-				Output.Errors.Add(FShaderCompilerError(FString::Printf(TEXT("Cross-compiled shader to intermediate HLSL as pre-compile step for FXC: %s"), *Input.GenerateShaderName())));
-			}
-
 			CrossCompiler::FShaderConductorContext CompilerContext;
 
 			// Load shader source into compiler context
@@ -804,9 +750,7 @@ bool CompileAndProcessD3DShaderFXC(FString& PreprocessedShaderSource, const FStr
 				CompileFlags & (~D3DCOMPILE_WARNINGS_ARE_ERRORS),
 				0,
 				Shader.GetInitReference(),
-				Errors.GetInitReference(),
-				bException,
-				ExceptionInfo
+				Errors.GetInitReference()
 			);
 
 			if (!bPrecompileWithDXC && SUCCEEDED(Result))
@@ -815,30 +759,6 @@ bool CompileAndProcessD3DShaderFXC(FString& PreprocessedShaderSource, const FStr
 				Output.Errors.Add(FShaderCompilerError(FString::Printf(TEXT("Cross-compiled shader to intermediate HLSL after first attempt crashed FXC: %s"), *Input.GenerateShaderName())));
 			}
 		}
-
-		if (bException)
-		{
-			const char* CodeName = Win32SehExceptionToString(ExceptionInfo.Code);
-			const FString CodeNameStr = (CodeName != nullptr ? ANSI_TO_TCHAR(CodeName) : TEXT("Unknown"));
-			FilteredErrors.Add(
-				FString::Printf(
-					TEXT("D3DCompile exception: Code = 0x%08X (%s), Address = 0x%016llX, Offset = 0x%016llX, Codebase = 0x%016llX"),
-					ExceptionInfo.Code,
-					*CodeNameStr,
-					ExceptionInfo.Address,
-					(ExceptionInfo.Address - ExceptionInfo.Base),
-					ExceptionInfo.Base
-				)
-			);
-
-			// Dump input shader source on exception to be able to investigate issue through logs on CIS servers
-			FString DumpedSource;
-			DumpedSource = PreprocessedShaderSource;
-			DumpedSource += TEXT("\n#if 0 /*DIRECT COMPILE*/\n");
-			DumpedSource += CreateShaderCompilerWorkerDirectCommandLine(Input);
-			DumpedSource += TEXT("\n#endif /*DIRECT COMPILE*/\n");
-			Output.OptionalPreprocessedShaderSource = MoveTemp(DumpedSource);
-		}
 	}
 	else
 	{
diff --git a/Engine/Source/Programs/ShaderCompileWorker/Private/ShaderCompileWorker.cpp b/Engine/Source/Programs/ShaderCompileWorker/Private/ShaderCompileWorker.cpp
--- a/Engine/Source/Programs/ShaderCompileWorker/Private/ShaderCompileWorker.cpp
+++ b/Engine/Source/Programs/ShaderCompileWorker/Private/ShaderCompileWorker.cpp
@@ -6,6 +6,7 @@
 
 #include "CoreMinimal.h"
 #include "RequiredProgramMainCPPInclude.h"
+#include "ShaderCompilerCore.h"
 #include "ShaderCore.h"
 #include "HAL/ExceptionHandling.h"
 #include "Interfaces/IShaderFormat.h"
@@ -107,6 +108,86 @@ static const IShaderFormat* FindShaderFormat(FName Name)
 	return nullptr;
 }
 
+#if PLATFORM_WINDOWS
+	static bool ExceptionCodeToString(DWORD ExceptionCode, FString& OutStr)
+	{
+#define EXCEPTION_CODE_CASE_STR(CODE) case CODE: OutStr = TEXT(#CODE); return true;
+		const DWORD CPlusPlusExceptionCode = 0xE06D7363;
+		switch (ExceptionCode)
+		{
+			EXCEPTION_CODE_CASE_STR(EXCEPTION_ACCESS_VIOLATION);
+			EXCEPTION_CODE_CASE_STR(EXCEPTION_ARRAY_BOUNDS_EXCEEDED);
+			EXCEPTION_CODE_CASE_STR(EXCEPTION_BREAKPOINT);
+			EXCEPTION_CODE_CASE_STR(EXCEPTION_DATATYPE_MISALIGNMENT);
+			EXCEPTION_CODE_CASE_STR(EXCEPTION_FLT_DENORMAL_OPERAND);
+			EXCEPTION_CODE_CASE_STR(EXCEPTION_FLT_DIVIDE_BY_ZERO);
+			EXCEPTION_CODE_CASE_STR(EXCEPTION_FLT_INEXACT_RESULT);
+			EXCEPTION_CODE_CASE_STR(EXCEPTION_FLT_INVALID_OPERATION);
+			EXCEPTION_CODE_CASE_STR(EXCEPTION_FLT_OVERFLOW);
+			EXCEPTION_CODE_CASE_STR(EXCEPTION_FLT_STACK_CHECK);
+			EXCEPTION_CODE_CASE_STR(EXCEPTION_FLT_UNDERFLOW);
+			EXCEPTION_CODE_CASE_STR(EXCEPTION_GUARD_PAGE);
+			EXCEPTION_CODE_CASE_STR(EXCEPTION_ILLEGAL_INSTRUCTION);
+			EXCEPTION_CODE_CASE_STR(EXCEPTION_IN_PAGE_ERROR);
+			EXCEPTION_CODE_CASE_STR(EXCEPTION_INT_DIVIDE_BY_ZERO);
+			EXCEPTION_CODE_CASE_STR(EXCEPTION_INT_OVERFLOW);
+			EXCEPTION_CODE_CASE_STR(EXCEPTION_INVALID_DISPOSITION);
+			EXCEPTION_CODE_CASE_STR(EXCEPTION_INVALID_HANDLE);
+			EXCEPTION_CODE_CASE_STR(EXCEPTION_NONCONTINUABLE_EXCEPTION);
+			EXCEPTION_CODE_CASE_STR(EXCEPTION_PRIV_INSTRUCTION);
+			EXCEPTION_CODE_CASE_STR(EXCEPTION_SINGLE_STEP);
+			EXCEPTION_CODE_CASE_STR(EXCEPTION_STACK_OVERFLOW);
+			EXCEPTION_CODE_CASE_STR(STATUS_UNWIND_CONSOLIDATE);
+			case CPlusPlusExceptionCode: OutStr = TEXT("CPP_EXCEPTION"); return true;
+			default: return false;
+		}
+#undef EXCEPTION_CODE_CASE_STR
+	}
+
+	static int HandleException(Windows::LPEXCEPTION_POINTERS Info, FString& OutExMsg, FString& OutCallStack)
+	{
+		const DWORD AssertExceptionCode = 0x00004000;
+		FString ExCodeStr;
+		if (Info->ExceptionRecord->ExceptionCode == AssertExceptionCode)
+		{
+			// In the case of an assert the assert handler populates the GErrorHist global.
+			// This contains a readable assert message followed by a callstack; so we can use that to populate
+			// our message/callstack and save some time as well as getting the properly formatted assert message.
+			FString Assert = GErrorHist;
+			const TCHAR* CallstackStart = FCString::Strfind(GErrorHist, TEXT("0x"));
+
+			OutExMsg = FString(CallstackStart - GErrorHist, GErrorHist);
+			OutCallStack = CallstackStart;
+		}
+		else
+		{
+			if (ExceptionCodeToString(Info->ExceptionRecord->ExceptionCode, ExCodeStr))
+			{
+				OutExMsg = FString::Printf(
+					TEXT("Exception: %s, address=0x%016x\n"),
+					*ExCodeStr,
+					(uint64)Info->ExceptionRecord->ExceptionAddress);
+			}
+			else
+			{
+				OutExMsg = FString::Printf(
+					TEXT("Exception code: 0x%08x, address=0x%016x\n"),
+					Info->ExceptionRecord->ExceptionCode,
+					(uint64)Info->ExceptionRecord->ExceptionAddress);
+			}
+
+			ANSICHAR CallStack[32768];
+			FMemory::Memzero(CallStack);
+			FPlatformStackWalk::StackWalkAndDump(CallStack, ARRAYSIZE(CallStack), Info->ExceptionRecord->ExceptionAddress);
+			OutCallStack = ANSI_TO_TCHAR(CallStack);
+		}
+
+		return EXCEPTION_EXECUTE_HANDLER;
+	}
+	
+#endif
+
+
 /** Processes a compilation job. */
 static void ProcessCompilationJob(const FShaderCompilerInput& Input,FShaderCompilerOutput& Output,const FString& WorkingDirectory)
 {
@@ -128,25 +209,45 @@ static void ProcessCompilationJob(const FShaderCompilerInput& Input,FShaderCompi
 		}
 	}
 
-	// Compile the shader directly through the platform dll (directly from the shader dir as the working directory)
-	double TimeStart = FPlatformTime::Seconds();
-	Compiler->CompileShader(Input.ShaderFormat, Input, Output, WorkingDirectory);
-	if (Output.bSucceeded)
+#if PLATFORM_WINDOWS
+	FString ExceptionMsg;
+	FString ExceptionCallstack;
+	__try
+#endif
 	{
-		Output.GenerateOutputHash();
-		if (Input.CompressionFormat != NAME_None)
+		// Compile the shader directly through the platform dll (directly from the shader dir as the working directory)
+		double TimeStart = FPlatformTime::Seconds();
+		Compiler->CompileShader(Input.ShaderFormat, Input, Output, WorkingDirectory);
+		if (Output.bSucceeded)
 		{
-			Output.CompressOutput(Input.CompressionFormat, Input.OodleCompressor, Input.OodleLevel);
+			Output.GenerateOutputHash();
+			if (Input.CompressionFormat != NAME_None)
+			{
+				Output.CompressOutput(Input.CompressionFormat, Input.OodleCompressor, Input.OodleLevel);
+			}
 		}
-	}
-	Output.CompileTime = FPlatformTime::Seconds() - TimeStart;
+		Output.CompileTime = FPlatformTime::Seconds() - TimeStart;
 
-	if (Compiler->UsesHLSLcc(Input))
+		if (Compiler->UsesHLSLcc(Input))
+		{
+			Output.bUsedHLSLccCompiler = true;
+		}
+
+		++GNumProcessedJobs;
+	}
+#if PLATFORM_WINDOWS
+	__except (HandleException(GetExceptionInformation(), ExceptionMsg, ExceptionCallstack))
 	{
-		Output.bUsedHLSLccCompiler = true;
+		Output.bSucceeded = false;
+
+		FShaderCompilerError ExceptionError;
+		ExceptionError.StrippedErrorMessage = FString::Printf(
+			TEXT("Exception encountered in platform compiler: %s\nException Callstack:\n%s"), 
+			*ExceptionMsg, 
+			*ExceptionCallstack);
+		Output.Errors.Add(ExceptionError);
 	}
-
-	++GNumProcessedJobs;
+#endif
 }
 
 static void UpdateFileSize(FArchive& OutputFile, int64 FileSizePosition)
@@ -1239,6 +1340,7 @@ static int32 GuardedMain(int32 argc, TCHAR* argv[], bool bDirectMode)
 	return 0;
 }
 
+
 static int32 GuardedMainWrapper(int32 ArgC, TCHAR* ArgV[], const TCHAR* CrashOutputFile, bool bDirectMode)
 {
 	FTaskTagScope Scope(ETaskTag::EGameThread);
@@ -1269,13 +1371,15 @@ static int32 GuardedMainWrapper(int32 ArgC, TCHAR* ArgV[], const TCHAR* CrashOut
 	{
 		// Don't want 32 dialogs popping up when SCW fails
 		GUseCrashReportClient = false;
+		FString ExceptionMsg;
+		FString ExceptionCallStack;
 		__try
 		{
 			GIsGuarded = 1;
 			ReturnCode = GuardedMain(ArgC, ArgV, bDirectMode);
 			GIsGuarded = 0;
 		}
-		__except(EXCEPTION_EXECUTE_HANDLER)
+		__except(HandleException(GetExceptionInformation(), ExceptionMsg, ExceptionCallStack))
 		{
 			FArchive& OutputFile = *IFileManager::Get().CreateFileWriter(CrashOutputFile, FILEWRITE_EvenIfReadOnly);
 
@@ -1292,8 +1396,9 @@ static int32 GuardedMainWrapper(int32 ArgC, TCHAR* ArgV[], const TCHAR* CrashOut
 					GFailedErrorCode = ESCWErrorCode::GeneralCrash;
 				}
 			}
-			int64 FileSizePosition = WriteOutputFileHeader(OutputFile, (int32)GFailedErrorCode, FCString::Strlen(GErrorHist), GErrorHist,
-				FCString::Strlen(GErrorExceptionDescription), GErrorExceptionDescription);
+			
+			int64 FileSizePosition = WriteOutputFileHeader(OutputFile, (int32)GFailedErrorCode, ExceptionCallStack.Len(), *ExceptionCallStack,
+				ExceptionMsg.Len(), *ExceptionMsg);
 
 			int32 NumBatches = 0;
 			OutputFile << NumBatches;
diff --git a/Engine/Source/Runtime/Engine/Private/ShaderCompiler/ShaderCompiler.cpp b/Engine/Source/Runtime/Engine/Private/ShaderCompiler/ShaderCompiler.cpp
--- a/Engine/Source/Runtime/Engine/Private/ShaderCompiler/ShaderCompiler.cpp
+++ b/Engine/Source/Runtime/Engine/Private/ShaderCompiler/ShaderCompiler.cpp
@@ -1233,7 +1233,7 @@ namespace SCWErrorCode
 	void HandleGeneralCrash(const TCHAR* ExceptionInfo, const TCHAR* Callstack)
 	{
 		GLog->Panic();
-		UE_LOG(LogShaderCompilers, Fatal, TEXT("ShaderCompileWorker crashed!\n%s\n\t%s"), ExceptionInfo, Callstack);
+		UE_LOG(LogShaderCompilers, Error, TEXT("ShaderCompileWorker crashed!\n%s\n%s"), ExceptionInfo, Callstack);
 	}
 
 	void HandleBadShaderFormatVersion(const TCHAR* Data)
