Added version to output when running SCW without arguments

diff --git a/Engine/Source/Programs/ShaderCompileWorker/Private/ShaderCompileWorker.cpp b/Engine/Source/Programs/ShaderCompileWorker/Private/ShaderCompileWorker.cpp
--- a/Engine/Source/Programs/ShaderCompileWorker/Private/ShaderCompileWorker.cpp
+++ b/Engine/Source/Programs/ShaderCompileWorker/Private/ShaderCompileWorker.cpp
@@ -1458,7 +1458,7 @@ INT32_MAIN_INT32_ARGC_TCHAR_ARGV()
 	{
 		if (ArgC < 6)
 		{
-			printf("ShaderCompileWorker is called by UnrealEditor, it requires specific command line arguments.\n");
+			printf("ShaderCompileWorker (v%d) is called by UnrealEditor, it requires specific command line arguments.\n", ShaderCompileWorkerOutputVersion);
 			return -1;
 		}
 
