[ShaderCompileWorker] minor fixes:
- allow 0 TimeToLive to be specified explicitly on commandline (useful when debugging single-job debug worker input files, without this it keeps running the same job repeatedly) and modify single job worker input commandline txt file to use this instead of 0.5f
- change config-conditional behaviour serializing to memory instead of file on disk to be conditional on an optional cmdline argument instead, and move file copy into the block that does the file writing (otherwise it repeatedly tries and fails to copy the nonexistent file)

diff --git a/Engine/Source/Programs/ShaderCompileWorker/Private/ShaderCompileWorker.cpp b/Engine/Source/Programs/ShaderCompileWorker/Private/ShaderCompileWorker.cpp
--- a/Engine/Source/Programs/ShaderCompileWorker/Private/ShaderCompileWorker.cpp
+++ b/Engine/Source/Programs/ShaderCompileWorker/Private/ShaderCompileWorker.cpp
@@ -301,6 +301,7 @@ class FWorkLoop
 public:
 	// If we have been idle for 20 seconds then exit. Can be overriden from the cmd line with -TimeToLive=N where N is in seconds (and a float value)
 	float TimeToLive = 20.0f;
+	bool DisableFileWrite = false;
 
 	FWorkLoop(const TCHAR* ParentProcessIdText,const TCHAR* InWorkingDirectory,const TCHAR* InInputFilename,const TCHAR* InOutputFilename, TMap<FString, uint32>& InFormatVersionMap)
 	:	ParentProcessId(FCString::Atoi(ParentProcessIdText))
@@ -317,12 +318,11 @@ public:
 		{
 			if (Switch.StartsWith(TEXT("TimeToLive=")))
 			{
-				float TokenTime = FCString::Atof(Switch.GetCharArray().GetData() + 11);
-				if (TokenTime > 0)
-				{
-					TimeToLive = TokenTime;
-					break;
-				}
+				TimeToLive = FCString::Atof(Switch.GetCharArray().GetData() + 11);
+			}
+			else if (Switch.Equals(TEXT("DisableFileWrite")))
+			{
+				DisableFileWrite = true;
 			}
 		}
 	}
@@ -358,23 +358,24 @@ public:
 			}
 
 			// Prepare for output
-#if UE_BUILD_DEBUG
-			TArray<uint8> MemBlock;
-			FMemoryWriter MemWriter(MemBlock);
-			FArchive* OutputFilePtr = &MemWriter;
-#else
-			FArchive* OutputFilePtr = CreateOutputArchive();
-			check(OutputFilePtr);
-#endif
-			WriteToOutputArchive(OutputFilePtr, SingleJobResults, PipelineJobResults);
-
-#if !UE_BUILD_DEBUG
-			// Close the output file.
-			delete OutputFilePtr;
-#endif
-
-			// Change the output file name to requested one
-			IFileManager::Get().Move(*OutputFilePath, *TempFilePath);
+			if (DisableFileWrite)
+			{
+				// write to in-memory bytestream instead for debugging purposes
+				TArray<uint8> MemBlock;
+				FMemoryWriter MemWriter(MemBlock);
+				WriteToOutputArchive(&MemWriter, SingleJobResults, PipelineJobResults);
+			}
+			else
+			{
+				FArchive* OutputFilePtr = CreateOutputArchive();
+				check(OutputFilePtr);
+				WriteToOutputArchive(OutputFilePtr, SingleJobResults, PipelineJobResults);
+				// Close the output file.
+				delete OutputFilePtr;
+
+				// Change the output file name to requested one
+				IFileManager::Get().Move(*OutputFilePath, *TempFilePath);
+			}
 
 			if (IsUsingXGE())
 			{
diff --git a/Engine/Source/Runtime/Engine/Private/ShaderCompiler/ShaderCompiler.cpp b/Engine/Source/Runtime/Engine/Private/ShaderCompiler/ShaderCompiler.cpp
--- a/Engine/Source/Runtime/Engine/Private/ShaderCompiler/ShaderCompiler.cpp
+++ b/Engine/Source/Runtime/Engine/Private/ShaderCompiler/ShaderCompiler.cpp
@@ -1590,7 +1590,7 @@ const TCHAR* DebugWorkerOutputFileName = TEXT("DebugSCW.out");
 
 FString CreateShaderCompilerWorkerDebugCommandLine(const FShaderCompileJob* Job)
 {
-	return FString::Printf(TEXT("\"%s\" 0 \"%s\" %s %s -TimeToLive 0.5f"), // 0 is parent PID, pass short TTL to prevent SCW from hanging around waiting for input after processing the single job
+	return FString::Printf(TEXT("\"%s\" 0 \"%s\" %s %s -TimeToLive=0.0f"), // 0 is parent PID, pass zero TTL to make SCW process the single job then exit
 		*Job->Input.DumpDebugInfoPath, // working directory for SCW
 		*Job->Input.DebugGroupName, // console window title
 		DebugWorkerInputFileName,
