[shaders] add option to SCW to keep input files rather than deleting them; modify the cmdline args generated for the "single job" debug worker input files to include this option so the debug input file won't be deleted after executing to completion.

diff --git a/Engine/Source/Programs/ShaderCompileWorker/Private/ShaderCompileWorker.cpp b/Engine/Source/Programs/ShaderCompileWorker/Private/ShaderCompileWorker.cpp
--- a/Engine/Source/Programs/ShaderCompileWorker/Private/ShaderCompileWorker.cpp
+++ b/Engine/Source/Programs/ShaderCompileWorker/Private/ShaderCompileWorker.cpp
@@ -302,6 +302,7 @@ public:
 	// If we have been idle for 20 seconds then exit. Can be overriden from the cmd line with -TimeToLive=N where N is in seconds (and a float value)
 	float TimeToLive = 20.0f;
 	bool DisableFileWrite = false;
+	bool KeepInput = false;
 
 	FWorkLoop(const TCHAR* ParentProcessIdText,const TCHAR* InWorkingDirectory,const TCHAR* InInputFilename,const TCHAR* InOutputFilename, TMap<FString, uint32>& InFormatVersionMap)
 	:	ParentProcessId(FCString::Atoi(ParentProcessIdText))
@@ -324,6 +325,10 @@ public:
 			{
 				DisableFileWrite = true;
 			}
+			else if (Switch.Equals(TEXT("KeepInput")))
+			{
+				KeepInput = true;
+			}
 		}
 	}
 
@@ -855,10 +860,10 @@ private:
 		bool bResult = false;
 
 		// It seems XGE does not support deleting files.
-		// Don't delete the input file if we are running under Incredibuild.
+		// Don't delete the input file if we are running under Incredibuild (or if the cmdline args explicitly told us to keep it).
 		// In xml mode, we signal completion by creating a zero byte "Success" file after the output file has been fully written.
 		// In intercept mode, completion is signaled by this process terminating.
-		if (!IsUsingXGE())
+		if (!IsUsingXGE() && !KeepInput)
 		{
 			do 
 			{
diff --git a/Engine/Source/Runtime/Engine/Private/ShaderCompiler/ShaderCompiler.cpp b/Engine/Source/Runtime/Engine/Private/ShaderCompiler/ShaderCompiler.cpp
--- a/Engine/Source/Runtime/Engine/Private/ShaderCompiler/ShaderCompiler.cpp
+++ b/Engine/Source/Runtime/Engine/Private/ShaderCompiler/ShaderCompiler.cpp
@@ -1590,7 +1590,8 @@ const TCHAR* DebugWorkerOutputFileName = TEXT("DebugSCW.out");
 
 FString CreateShaderCompilerWorkerDebugCommandLine(const FShaderCompileJob* Job)
 {
-	return FString::Printf(TEXT("\"%s\" 0 \"%s\" %s %s -TimeToLive=0.0f"), // 0 is parent PID, pass zero TTL to make SCW process the single job then exit
+	// 0 is parent PID, pass zero TTL and KeepInput to make SCW process the single job then exit without deleting the input file
+	return FString::Printf(TEXT("\"%s\" 0 \"%s\" %s %s -TimeToLive=0.0f -KeepInput"),
 		*Job->Input.DumpDebugInfoPath, // working directory for SCW
 		*Job->Input.DebugGroupName, // console window title
 		DebugWorkerInputFileName,
