Bump GUID in ShaderVersion.ush to invalidate shader DDC and fix SCW output file protocol on the non-error path. Appendix to CL 23427669.

#rnx

diff --git a/Engine/Shaders/Public/ShaderVersion.ush b/Engine/Shaders/Public/ShaderVersion.ush
--- a/Engine/Shaders/Public/ShaderVersion.ush
+++ b/Engine/Shaders/Public/ShaderVersion.ush
@@ -2,4 +2,4 @@
 // Each time the console command is executed it generates a new GUID. As this file is included
 // in Platform.ush (which should be included in any shader) it allows to invalidate the shader DDC.
 // 
-// GUID = 28A1A473-9B5D-40E9-B951-B09A08EB4DA3
+// GUID = 8AD13364-EC26-47A7-964D-88C0AB4B126C
diff --git a/Engine/Source/Programs/ShaderCompileWorker/Private/ShaderCompileWorker.cpp b/Engine/Source/Programs/ShaderCompileWorker/Private/ShaderCompileWorker.cpp
--- a/Engine/Source/Programs/ShaderCompileWorker/Private/ShaderCompileWorker.cpp
+++ b/Engine/Source/Programs/ShaderCompileWorker/Private/ShaderCompileWorker.cpp
@@ -307,27 +307,30 @@ static int64 WriteOutputFileHeader(FArchive& OutputFile, int32 ErrorCode, int32
 	const TCHAR* Hostname = GetLocalHostname(&HostnameLength);
 	OutputFile << HostnameLength;
 
-	if (CallstackLength > 0)
+	if (ErrorCode != FSCWErrorCode::Success)
 	{
-		OutputFile.Serialize((void*)Callstack, CallstackLength * sizeof(TCHAR));
-	}
+		if (CallstackLength > 0)
+		{
+			OutputFile.Serialize((void*)Callstack, CallstackLength * sizeof(TCHAR));
+		}
 
-	if (ExceptionInfoLength > 0)
-	{
-		OutputFile.Serialize((void*)ExceptionInfo, ExceptionInfoLength * sizeof(TCHAR));
-	}
+		if (ExceptionInfoLength > 0)
+		{
+			OutputFile.Serialize((void*)ExceptionInfo, ExceptionInfoLength * sizeof(TCHAR));
+		}
 
-	if (HostnameLength > 0)
-	{
-		OutputFile.Serialize((void*)Hostname, HostnameLength * sizeof(TCHAR));
-	}
+		if (HostnameLength > 0)
+		{
+			OutputFile.Serialize((void*)Hostname, HostnameLength * sizeof(TCHAR));
+		}
 
-	// Store available and used physical memory of host machine on OOM error
-	if (ErrorCode == FSCWErrorCode::OutOfMemory)
-	{
-		uint64 AvailablePhysical = 0, UsedPhysical = 0;
-		GetMemoryStats(AvailablePhysical, UsedPhysical);
-		OutputFile << AvailablePhysical << UsedPhysical;
+		// Store available and used physical memory of host machine on OOM error
+		if (ErrorCode == FSCWErrorCode::OutOfMemory)
+		{
+			uint64 AvailablePhysical = 0, UsedPhysical = 0;
+			GetMemoryStats(AvailablePhysical, UsedPhysical);
+			OutputFile << AvailablePhysical << UsedPhysical;
+		}
 	}
 
 	UpdateFileSize(OutputFile, FileSizePosition);
