[shaders] factor out single-input compilation and exception handling into shared utility functions

diff --git a/Engine/Source/Developer/TargetPlatform/Private/TargetPlatformManagerModule.cpp b/Engine/Source/Developer/TargetPlatform/Private/TargetPlatformManagerModule.cpp
--- a/Engine/Source/Developer/TargetPlatform/Private/TargetPlatformManagerModule.cpp
+++ b/Engine/Source/Developer/TargetPlatform/Private/TargetPlatformManagerModule.cpp
@@ -7,6 +7,7 @@
 #include "Misc/ConfigCacheIni.h"
 #include "Misc/Paths.h"
 #include "Misc/OutputDeviceRedirector.h"
+#include "ShaderCompilerCore.h"
 #include "Stats/Stats.h"
 #include "Misc/ScopedSlowTask.h"
 #include "Misc/MonitoredProcess.h"
@@ -629,24 +630,7 @@ public:
 
 	virtual const IShaderFormat* FindShaderFormat(FName Name) override
 	{
-		const TArray<const IShaderFormat*>& ShaderFormats = GetShaderFormats();	
-
-		for (int32 Index = 0; Index < ShaderFormats.Num(); Index++)
-		{
-			TArray<FName> Formats;
-			
-			ShaderFormats[Index]->GetSupportedFormats(Formats);
-		
-			for (int32 FormatIndex = 0; FormatIndex < Formats.Num(); FormatIndex++)
-			{
-				if (Formats[FormatIndex] == Name)
-				{
-					return ShaderFormats[Index];
-				}
-			}
-		}
-
-		return nullptr;
+		return ::FindShaderFormat(Name, GetShaderFormats());
 	}
 
 	virtual uint32 ShaderFormatVersion(FName Name) override
diff --git a/Engine/Source/Programs/ShaderCompileWorker/Private/ShaderCompileWorker.cpp b/Engine/Source/Programs/ShaderCompileWorker/Private/ShaderCompileWorker.cpp
--- a/Engine/Source/Programs/ShaderCompileWorker/Private/ShaderCompileWorker.cpp
+++ b/Engine/Source/Programs/ShaderCompileWorker/Private/ShaderCompileWorker.cpp
@@ -7,6 +7,7 @@
 #include "CoreMinimal.h"
 #include "RequiredProgramMainCPPInclude.h"
 #include "ShaderCompilerCore.h"
+#include "ShaderCompilerCommon.h"
 #include "ShaderCore.h"
 #include "HAL/ExceptionHandling.h"
 #include "Interfaces/IShaderFormat.h"
@@ -89,157 +90,6 @@ static const TArray<const IShaderFormat*>& GetShaderFormats()
 	return Results;
 }
 
-static const IShaderFormat* FindShaderFormat(FName Name)
-{
-	const TArray<const IShaderFormat*>& ShaderFormats = GetShaderFormats();	
-
-	for (int32 Index = 0; Index < ShaderFormats.Num(); Index++)
-	{
-		TArray<FName> Formats;
-		ShaderFormats[Index]->GetSupportedFormats(Formats);
-		for (int32 FormatIndex = 0; FormatIndex < Formats.Num(); FormatIndex++)
-		{
-			if (Formats[FormatIndex] == Name)
-			{
-				return ShaderFormats[Index];
-			}
-		}
-	}
-
-	return nullptr;
-}
-
-#if PLATFORM_WINDOWS
-	static bool ExceptionCodeToString(DWORD ExceptionCode, FString& OutStr)
-	{
-#define EXCEPTION_CODE_CASE_STR(CODE) case CODE: OutStr = TEXT(#CODE); return true;
-		const DWORD CPlusPlusExceptionCode = 0xE06D7363;
-		switch (ExceptionCode)
-		{
-			EXCEPTION_CODE_CASE_STR(EXCEPTION_ACCESS_VIOLATION);
-			EXCEPTION_CODE_CASE_STR(EXCEPTION_ARRAY_BOUNDS_EXCEEDED);
-			EXCEPTION_CODE_CASE_STR(EXCEPTION_BREAKPOINT);
-			EXCEPTION_CODE_CASE_STR(EXCEPTION_DATATYPE_MISALIGNMENT);
-			EXCEPTION_CODE_CASE_STR(EXCEPTION_FLT_DENORMAL_OPERAND);
-			EXCEPTION_CODE_CASE_STR(EXCEPTION_FLT_DIVIDE_BY_ZERO);
-			EXCEPTION_CODE_CASE_STR(EXCEPTION_FLT_INEXACT_RESULT);
-			EXCEPTION_CODE_CASE_STR(EXCEPTION_FLT_INVALID_OPERATION);
-			EXCEPTION_CODE_CASE_STR(EXCEPTION_FLT_OVERFLOW);
-			EXCEPTION_CODE_CASE_STR(EXCEPTION_FLT_STACK_CHECK);
-			EXCEPTION_CODE_CASE_STR(EXCEPTION_FLT_UNDERFLOW);
-			EXCEPTION_CODE_CASE_STR(EXCEPTION_GUARD_PAGE);
-			EXCEPTION_CODE_CASE_STR(EXCEPTION_ILLEGAL_INSTRUCTION);
-			EXCEPTION_CODE_CASE_STR(EXCEPTION_IN_PAGE_ERROR);
-			EXCEPTION_CODE_CASE_STR(EXCEPTION_INT_DIVIDE_BY_ZERO);
-			EXCEPTION_CODE_CASE_STR(EXCEPTION_INT_OVERFLOW);
-			EXCEPTION_CODE_CASE_STR(EXCEPTION_INVALID_DISPOSITION);
-			EXCEPTION_CODE_CASE_STR(EXCEPTION_INVALID_HANDLE);
-			EXCEPTION_CODE_CASE_STR(EXCEPTION_NONCONTINUABLE_EXCEPTION);
-			EXCEPTION_CODE_CASE_STR(EXCEPTION_PRIV_INSTRUCTION);
-			EXCEPTION_CODE_CASE_STR(EXCEPTION_SINGLE_STEP);
-			EXCEPTION_CODE_CASE_STR(EXCEPTION_STACK_OVERFLOW);
-			EXCEPTION_CODE_CASE_STR(STATUS_UNWIND_CONSOLIDATE);
-			case CPlusPlusExceptionCode: OutStr = TEXT("CPP_EXCEPTION"); return true;
-			default: return false;
-		}
-#undef EXCEPTION_CODE_CASE_STR
-	}
-
-	static int HandleException(Windows::LPEXCEPTION_POINTERS Info, FString& OutExMsg, FString& OutCallStack)
-	{
-		const DWORD AssertExceptionCode = 0x00004000;
-		FString ExCodeStr;
-		if (Info->ExceptionRecord->ExceptionCode == AssertExceptionCode)
-		{
-			// In the case of an assert the assert handler populates the GErrorHist global.
-			// This contains a readable assert message followed by a callstack; so we can use that to populate
-			// our message/callstack and save some time as well as getting the properly formatted assert message.
-			FString Assert = GErrorHist;
-			const TCHAR* CallstackStart = FCString::Strfind(GErrorHist, TEXT("0x"));
-
-			OutExMsg = FString(CallstackStart - GErrorHist, GErrorHist);
-			OutCallStack = CallstackStart;
-		}
-		else
-		{
-			if (ExceptionCodeToString(Info->ExceptionRecord->ExceptionCode, ExCodeStr))
-			{
-				OutExMsg = FString::Printf(
-					TEXT("Exception: %s, address=0x%016x\n"),
-					*ExCodeStr,
-					(uint64)Info->ExceptionRecord->ExceptionAddress);
-			}
-			else
-			{
-				OutExMsg = FString::Printf(
-					TEXT("Exception code: 0x%08x, address=0x%016x\n"),
-					Info->ExceptionRecord->ExceptionCode,
-					(uint64)Info->ExceptionRecord->ExceptionAddress);
-			}
-
-			ANSICHAR CallStack[32768];
-			FMemory::Memzero(CallStack);
-			FPlatformStackWalk::StackWalkAndDump(CallStack, ARRAYSIZE(CallStack), Info->ExceptionRecord->ExceptionAddress);
-			OutCallStack = ANSI_TO_TCHAR(CallStack);
-		}
-
-		return EXCEPTION_EXECUTE_HANDLER;
-	}
-	
-#endif
-
-
-/** Processes a compilation job. */
-static FORCENOINLINE void ProcessCompilationJob(const FShaderCompilerInput& Input,FShaderCompilerOutput& Output,const FString& WorkingDirectory)
-{
-	TRACE_CPUPROFILER_EVENT_SCOPE(ProcessCompilationJob);
-
-	const IShaderFormat* Compiler = FindShaderFormat(Input.ShaderFormat);
-	if (!Compiler)
-	{
-		ExitWithoutCrash(FSCWErrorCode::CantCompileForSpecificFormat, FString::Printf(TEXT("Can't compile shaders for format %s"), *Input.ShaderFormat.ToString()));
-	}
-
-#if PLATFORM_WINDOWS
-	FString ExceptionMsg;
-	FString ExceptionCallstack;
-	__try
-#endif
-	{
-		// Compile the shader directly through the platform dll (directly from the shader dir as the working directory)
-		double TimeStart = FPlatformTime::Seconds();
-		Compiler->CompileShader(Input.ShaderFormat, Input, Output, WorkingDirectory);
-		if (Output.bSucceeded)
-		{
-			Output.GenerateOutputHash();
-			if (Input.CompressionFormat != NAME_None)
-			{
-				Output.CompressOutput(Input.CompressionFormat, Input.OodleCompressor, Input.OodleLevel);
-			}
-		}
-		Output.CompileTime = FPlatformTime::Seconds() - TimeStart;
-
-		if (Compiler->UsesHLSLcc(Input))
-		{
-			Output.bUsedHLSLccCompiler = true;
-		}
-
-		++GNumProcessedJobs;
-	}
-#if PLATFORM_WINDOWS
-	__except (HandleException(GetExceptionInformation(), ExceptionMsg, ExceptionCallstack))
-	{
-		Output.bSucceeded = false;
-
-		FShaderCompilerError ExceptionError;
-		ExceptionError.StrippedErrorMessage = FString::Printf(
-			TEXT("Exception encountered in platform compiler: %s\nException Callstack:\n%s"), 
-			*ExceptionMsg, 
-			*ExceptionCallstack);
-		Output.Errors.Add(ExceptionError);
-	}
-#endif
-}
 
 static void UpdateFileSize(FArchive& OutputFile, int64 FileSizePosition)
 {
@@ -773,15 +623,9 @@ private:
 					UpdateIncludeDirectoryForPreviewPlatform((EShaderPlatform)CompilerInput.Target.Platform, ShaderPlatform);
 				}
 
-				if (IsValidRef(CompilerInput.SharedEnvironment))
-				{
-					// Merge the shared environment into the per-shader environment before calling into the compile function
-					CompilerInput.Environment.Merge(*CompilerInput.SharedEnvironment);
-				}
-
 				// Process the job.
 				FShaderCompilerOutput CompilerOutput;
-				ProcessCompilationJob(CompilerInput, CompilerOutput, WorkingDirectory);
+				CompileShader(GetShaderFormats(), CompilerInput, CompilerOutput, WorkingDirectory, &GNumProcessedJobs);
 
 				// Serialize the job's output.
 				FJobResult& JobResult = *new(OutSingleJobResults) FJobResult;
@@ -826,12 +670,6 @@ private:
 						const EShaderPlatform ShaderPlatform = ShaderFormatNameToShaderPlatform(CompilerInputs[StageIndex].ShaderFormat);
 						UpdateIncludeDirectoryForPreviewPlatform((EShaderPlatform)CompilerInputs[StageIndex].Target.Platform, ShaderPlatform);
 					}
-
-					if (IsValidRef(CompilerInputs[StageIndex].SharedEnvironment))
-					{
-						// Merge the shared environment into the per-shader environment before calling into the compile function
-						CompilerInputs[StageIndex].Environment.Merge(*CompilerInputs[StageIndex].SharedEnvironment);
-					}
 				}
 
 				ProcessShaderPipelineCompilationJob(PipelineJob, CompilerInputs);
@@ -847,7 +685,7 @@ private:
 		FShaderCompilerOutput FirstCompilerOutput;
 		CompilerInputs[0].bCompilingForShaderPipeline = true;
 		CompilerInputs[0].bIncludeUsedOutputs = false;
-		ProcessCompilationJob(CompilerInputs[0], FirstCompilerOutput, WorkingDirectory);
+		CompileShader(GetShaderFormats(), CompilerInputs[0], FirstCompilerOutput, WorkingDirectory, &GNumProcessedJobs);
 
 		// Serialize the job's output.
 		{
@@ -878,7 +716,7 @@ private:
 			}
 
 			FShaderCompilerOutput CompilerOutput;
-			ProcessCompilationJob(CompilerInputs[Index], CompilerOutput, WorkingDirectory);
+			CompileShader(GetShaderFormats(), CompilerInputs[Index], CompilerOutput, WorkingDirectory, &GNumProcessedJobs);
 
 			// Serialize the job's output.
 			FJobResult& JobResult = *new(PipelineJob.SingleJobs) FJobResult;
@@ -1233,7 +1071,7 @@ static void DirectCompile(const TArray<const class IShaderFormat*>& ShaderFormat
 	Input.UsedOutputs = UsedOutputs;
 
 	FShaderCompilerOutput Output;
-	ProcessCompilationJob(Input, Output, Dir);
+	CompileShader(GetShaderFormats(), Input, Output, Dir, &GNumProcessedJobs);
 }
 
 
@@ -1428,7 +1266,7 @@ static int32 GuardedMainWrapper(int32 ArgC, TCHAR* ArgV[], const TCHAR* CrashOut
 			ReturnCode = GuardedMain(ArgC, ArgV, bDirectMode);
 			GIsGuarded = 0;
 		}
-		__except(HandleException(GetExceptionInformation(), ExceptionMsg, ExceptionCallStack))
+		__except(HandleShaderCompileException(GetExceptionInformation(), ExceptionMsg, ExceptionCallStack))
 		{
 			FArchive& OutputFile = *IFileManager::Get().CreateFileWriter(CrashOutputFile, FILEWRITE_EvenIfReadOnly);
 
diff --git a/Engine/Source/Runtime/Engine/Private/ShaderCompiler/ShaderCompiler.cpp b/Engine/Source/Runtime/Engine/Private/ShaderCompiler/ShaderCompiler.cpp
--- a/Engine/Source/Runtime/Engine/Private/ShaderCompiler/ShaderCompiler.cpp
+++ b/Engine/Source/Runtime/Engine/Private/ShaderCompiler/ShaderCompiler.cpp
@@ -1193,7 +1193,7 @@ static TAutoConsoleVariable<bool> CVarShadersUseLegacyPreprocessor(
 	TEXT("\tfalse: Disabled - preprocess with STB\n"),
 	ECVF_ReadOnly);
 
-extern bool CompileShaderPipeline(const IShaderFormat* Compiler, FName Format, FShaderPipelineCompileJob* PipelineJob, const FString& Dir);
+extern bool CompileShaderPipeline(const TArray<const IShaderFormat*>& ShaderFormats, FName Format, FShaderPipelineCompileJob* PipelineJob, const FString& Dir);
 
 #if ENABLE_COOK_STATS
 namespace ShaderCompilerCookStats
@@ -2910,41 +2910,14 @@ void FShaderCompileUtilities::ExecuteShaderCompileJob(FShaderCommonCompileJob& J
 
 	check(!Job.bFinalized);
 
+	FString WorkingDir = FPlatformProcess::ShaderDir();
 	static ITargetPlatformManagerModule& TPM = GetTargetPlatformManagerRef();
 	auto* SingleJob = Job.GetSingleShaderJob();
+	const TArray<const IShaderFormat*> ShaderFormats = TPM.GetShaderFormats();
 	if (SingleJob)
 	{
 		const FName Format = (SingleJob->Input.ShaderFormat != NAME_None) ? SingleJob->Input.ShaderFormat : LegacyShaderPlatformToShaderFormat(EShaderPlatform(SingleJob->Input.Target.Platform));
-		const IShaderFormat* Compiler = TPM.FindShaderFormat(Format);
-
-		if (!Compiler)
-		{
-			UE_LOG(LogShaderCompilers, Fatal, TEXT("Can't compile shaders for format %s, couldn't load compiler dll"), *Format.ToString());
-		}
-		CA_ASSUME(Compiler != nullptr);
-
-		if (IsValidRef(SingleJob->Input.SharedEnvironment))
-		{
-			// Merge the shared environment into the per-shader environment before calling into the compile function
-			// Normally this happens in the worker
-			SingleJob->Input.Environment.Merge(*SingleJob->Input.SharedEnvironment);
-		}
-
-		// Compile the shader directly through the platform dll (directly from the shader dir as the working directory)
-		Compiler->CompileShader(Format, SingleJob->Input, SingleJob->Output, FString(FPlatformProcess::ShaderDir()));
-
-		SingleJob->bSucceeded = SingleJob->Output.bSucceeded;
-
-		if (SingleJob->Output.bSucceeded)
-		{
-			// Generate a hash of the output and compress the code
-			// The shader processing this output will use the heash to search for existing FShaderResources
-			SingleJob->Output.GenerateOutputHash();
-			if (SingleJob->Input.CompressionFormat != NAME_None)
-			{
-				SingleJob->Output.CompressOutput(SingleJob->Input.CompressionFormat, SingleJob->Input.OodleCompressor, SingleJob->Input.OodleLevel);
-			}
-		}
+		CompileShader(ShaderFormats, SingleJob->Input, SingleJob->Output, WorkingDir);
 	}
 	else
 	{
@@ -2953,13 +2926,6 @@ void FShaderCompileUtilities::ExecuteShaderCompileJob(FShaderCommonCompileJob& J
 
 		EShaderPlatform Platform = (EShaderPlatform)PipelineJob->StageJobs[0]->Input.Target.Platform;
 		const FName Format = LegacyShaderPlatformToShaderFormat(Platform);
-		const IShaderFormat* Compiler = TPM.FindShaderFormat(Format);
-
-		if (!Compiler)
-		{
-			UE_LOG(LogShaderCompilers, Fatal, TEXT("Can't compile shaders for format %s, couldn't load compiler dll"), *Format.ToString());
-		}
-		CA_ASSUME(Compiler != nullptr);
 
 		// Verify same platform on all stages
 		for (int32 Index = 1; Index < PipelineJob->StageJobs.Num(); ++Index)
@@ -2975,7 +2941,7 @@ void FShaderCompileUtilities::ExecuteShaderCompileJob(FShaderCommonCompileJob& J
 			}
 		}
 
-		CompileShaderPipeline(Compiler, Format, PipelineJob, FString(FPlatformProcess::ShaderDir()));
+		CompileShaderPipeline(ShaderFormats, Format, PipelineJob, WorkingDir);
 	}
 
 	Job.bFinalized = true;
diff --git a/Engine/Source/Runtime/Engine/Private/ShaderCompiler/ShaderPipelineCompiler.cpp b/Engine/Source/Runtime/Engine/Private/ShaderCompiler/ShaderPipelineCompiler.cpp
--- a/Engine/Source/Runtime/Engine/Private/ShaderCompiler/ShaderPipelineCompiler.cpp
+++ b/Engine/Source/Runtime/Engine/Private/ShaderCompiler/ShaderPipelineCompiler.cpp
@@ -7,7 +7,7 @@
 #include "Interfaces/IShaderFormat.h"
 #include "ShaderCompiler.h"
 
-bool CompileShaderPipeline(const IShaderFormat* Compiler, FName Format, FShaderPipelineCompileJob* PipelineJob, const FString& Dir)
+bool CompileShaderPipeline(const TArray<const IShaderFormat*>& ShaderFormats, FName Format, FShaderPipelineCompileJob* PipelineJob, const FString& Dir)
 {
 	check(PipelineJob && PipelineJob->StageJobs.Num() > 0);
 
@@ -18,28 +18,16 @@ bool CompileShaderPipeline(const IShaderFormat* Compiler, FName Format, FShaderP
 	// First job doesn't have to trim outputs
 	CurrentJob->Input.bIncludeUsedOutputs = false;
 
-	if (IsValidRef(CurrentJob->Input.SharedEnvironment))
-	{
-		// Merge the shared environment into the per-shader environment before calling into the compile function
-		// Normally this happens in the worker
-		CurrentJob->Input.Environment.Merge(*CurrentJob->Input.SharedEnvironment);
-	}
-
 	// Compile the shader directly through the platform dll (directly from the shader dir as the working directory)
-	Compiler->CompileShader(Format, CurrentJob->Input, CurrentJob->Output, Dir);
+	CompileShader(ShaderFormats, CurrentJob->Input, CurrentJob->Output, Dir);
 
 	CurrentJob->bSucceeded = CurrentJob->Output.bSucceeded;
-
 	if (!CurrentJob->Output.bSucceeded)
 	{
 		// Can't carry on compiling the pipeline
 		return false;
 	}
 
-	// Generate a hash of the output and cache it
-	// The shader processing this output will use it to search for existing FShaderResources
-	CurrentJob->Output.GenerateOutputHash();
-
 	// This tells the shader compiler we do want to remove unused outputs
 	bool bEnableRemovingUnused = true;
 	// This tells us the hlsl parser failed at removing unused outputs
@@ -68,18 +56,10 @@ bool CompileShaderPipeline(const IShaderFormat* Compiler, FName Format, FShaderP
 			CurrentJob->Input.UsedOutputs = PreviousJob->Output.UsedAttributes;
 		}
 
-		if (IsValidRef(CurrentJob->Input.SharedEnvironment))
-		{
-			// Merge the shared environment into the per-shader environment before calling into the compile function
-			// Normally this happens in the worker
-			CurrentJob->Input.Environment.Merge(*CurrentJob->Input.SharedEnvironment);
-		}
-
 		// Compile the shader directly through the platform dll (directly from the shader dir as the working directory)
-		Compiler->CompileShader(Format, CurrentJob->Input, CurrentJob->Output, Dir);
+		CompileShader(ShaderFormats, CurrentJob->Input, CurrentJob->Output, Dir);
 
 		CurrentJob->bSucceeded = CurrentJob->Output.bSucceeded;
-
 		if (!CurrentJob->Output.bSucceeded)
 		{
 			// Can't carry on compiling the pipeline
@@ -87,10 +67,6 @@ bool CompileShaderPipeline(const IShaderFormat* Compiler, FName Format, FShaderP
 		}
 
 		bJobFailedRemovingUnused = CurrentJob->Output.bFailedRemovingUnused || bJobFailedRemovingUnused;
-
-		// Generate a hash of the output and cache it
-		// The shader processing this output will use it to search for existing FShaderResources
-		CurrentJob->Output.GenerateOutputHash();
 	}
 
 	PipelineJob->bSucceeded = true;
diff --git a/Engine/Source/Runtime/RenderCore/Private/ShaderCore.cpp b/Engine/Source/Runtime/RenderCore/Private/ShaderCore.cpp
--- a/Engine/Source/Runtime/RenderCore/Private/ShaderCore.cpp
+++ b/Engine/Source/Runtime/RenderCore/Private/ShaderCore.cpp
@@ -7,6 +7,7 @@
 #include "ShaderCore.h"
 #include "Algo/Find.h"
 #include "HAL/FileManager.h"
+#include "HAL/PlatformStackWalk.h"
 #include "Misc/FileHelper.h"
 #include "Misc/Paths.h"
 #include "Misc/PathViews.h"
@@ -35,6 +36,12 @@
 #include "Misc/StringBuilder.h"
 #endif
 
+#if PLATFORM_WINDOWS
+#include "Windows/AllowWindowsPlatformTypes.h"
+#	include <winnt.h>
+#include "Windows/HideWindowsPlatformTypes.h"
+#endif
+
 static TAutoConsoleVariable<int32> CVarShaderDevelopmentMode(
 	TEXT("r.ShaderDevelopmentMode"),
 	0,
@@ -699,6 +706,178 @@ bool CheckVirtualShaderFilePath(FStringView VirtualFilePath, TArray<FShaderCompi
 	return bSuccess;
 }
 
+const IShaderFormat* FindShaderFormat(FName Format, TArray<const IShaderFormat*> ShaderFormats)
+{
+	for (int32 Index = 0; Index < ShaderFormats.Num(); Index++)
+	{
+		TArray<FName> Formats;
+
+		ShaderFormats[Index]->GetSupportedFormats(Formats);
+
+		for (int32 FormatIndex = 0; FormatIndex < Formats.Num(); FormatIndex++)
+		{
+			if (Formats[FormatIndex] == Format)
+			{
+				return ShaderFormats[Index];
+			}
+		}
+	}
+
+	return nullptr;
+}
+
+#if PLATFORM_WINDOWS
+static bool ExceptionCodeToString(DWORD ExceptionCode, FString& OutStr)
+{
+#define EXCEPTION_CODE_CASE_STR(CODE) case CODE: OutStr = TEXT(#CODE); return true;
+	const DWORD CPlusPlusExceptionCode = 0xE06D7363;
+	switch (ExceptionCode)
+	{
+		EXCEPTION_CODE_CASE_STR(EXCEPTION_ACCESS_VIOLATION);
+		EXCEPTION_CODE_CASE_STR(EXCEPTION_ARRAY_BOUNDS_EXCEEDED);
+		EXCEPTION_CODE_CASE_STR(EXCEPTION_BREAKPOINT);
+		EXCEPTION_CODE_CASE_STR(EXCEPTION_DATATYPE_MISALIGNMENT);
+		EXCEPTION_CODE_CASE_STR(EXCEPTION_FLT_DENORMAL_OPERAND);
+		EXCEPTION_CODE_CASE_STR(EXCEPTION_FLT_DIVIDE_BY_ZERO);
+		EXCEPTION_CODE_CASE_STR(EXCEPTION_FLT_INEXACT_RESULT);
+		EXCEPTION_CODE_CASE_STR(EXCEPTION_FLT_INVALID_OPERATION);
+		EXCEPTION_CODE_CASE_STR(EXCEPTION_FLT_OVERFLOW);
+		EXCEPTION_CODE_CASE_STR(EXCEPTION_FLT_STACK_CHECK);
+		EXCEPTION_CODE_CASE_STR(EXCEPTION_FLT_UNDERFLOW);
+		EXCEPTION_CODE_CASE_STR(EXCEPTION_GUARD_PAGE);
+		EXCEPTION_CODE_CASE_STR(EXCEPTION_ILLEGAL_INSTRUCTION);
+		EXCEPTION_CODE_CASE_STR(EXCEPTION_IN_PAGE_ERROR);
+		EXCEPTION_CODE_CASE_STR(EXCEPTION_INT_DIVIDE_BY_ZERO);
+		EXCEPTION_CODE_CASE_STR(EXCEPTION_INT_OVERFLOW);
+		EXCEPTION_CODE_CASE_STR(EXCEPTION_INVALID_DISPOSITION);
+		EXCEPTION_CODE_CASE_STR(EXCEPTION_INVALID_HANDLE);
+		EXCEPTION_CODE_CASE_STR(EXCEPTION_NONCONTINUABLE_EXCEPTION);
+		EXCEPTION_CODE_CASE_STR(EXCEPTION_PRIV_INSTRUCTION);
+		EXCEPTION_CODE_CASE_STR(EXCEPTION_SINGLE_STEP);
+		EXCEPTION_CODE_CASE_STR(EXCEPTION_STACK_OVERFLOW);
+		EXCEPTION_CODE_CASE_STR(STATUS_UNWIND_CONSOLIDATE);
+		case CPlusPlusExceptionCode: OutStr = TEXT("CPP_EXCEPTION"); return true;
+		default: return false;
+	}
+#undef EXCEPTION_CODE_CASE_STR
+}
+
+int HandleShaderCompileException(Windows::LPEXCEPTION_POINTERS Info, FString& OutExMsg, FString& OutCallStack)
+{
+	const DWORD AssertExceptionCode = 0x00004000;
+	FString ExCodeStr;
+	if (Info->ExceptionRecord->ExceptionCode == AssertExceptionCode)
+	{
+		// In the case of an assert the assert handler populates the GErrorHist global.
+		// This contains a readable assert message followed by a callstack; so we can use that to populate
+		// our message/callstack and save some time as well as getting the properly formatted assert message.
+		FString Assert = GErrorHist;
+		const TCHAR* CallstackStart = FCString::Strfind(GErrorHist, TEXT("0x"));
+
+		OutExMsg = FString(CallstackStart - GErrorHist, GErrorHist);
+		OutCallStack = CallstackStart;
+	}
+	else
+	{
+		if (ExceptionCodeToString(Info->ExceptionRecord->ExceptionCode, ExCodeStr))
+		{
+			OutExMsg = FString::Printf(
+				TEXT("Exception: %s, address=0x%016x\n"),
+				*ExCodeStr,
+				(uint64)Info->ExceptionRecord->ExceptionAddress);
+		}
+		else
+		{
+			OutExMsg = FString::Printf(
+				TEXT("Exception code: 0x%08x, address=0x%016x\n"),
+				Info->ExceptionRecord->ExceptionCode,
+				(uint64)Info->ExceptionRecord->ExceptionAddress);
+		}
+
+		ANSICHAR CallStack[32768];
+		FMemory::Memzero(CallStack);
+		FPlatformStackWalk::StackWalkAndDump(CallStack, ARRAYSIZE(CallStack), Info->ExceptionRecord->ExceptionAddress);
+		OutCallStack = ANSI_TO_TCHAR(CallStack);
+	}
+
+	return EXCEPTION_EXECUTE_HANDLER;
+}
+#endif
+
+#if PLATFORM_WINDOWS
+void CompileShaderInternal(const IShaderFormat* Compiler, const FShaderCompilerInput& Input, FShaderCompilerOutput& Output, const FString& WorkingDirectory, FString& OutExceptionMsg, FString& OutExceptionCallstack, int32* CompileCount)
+#else
+void CompileShaderInternal(const IShaderFormat* Compiler, const FShaderCompilerInput& Input, FShaderCompilerOutput& Output, const FString& WorkingDirectory, int32* CompileCount)
+#endif
+{
+#if PLATFORM_WINDOWS
+	__try
+#endif
+	{
+		double TimeStart = FPlatformTime::Seconds();
+
+		Compiler->CompileShader(Input.ShaderFormat, Input, Output, WorkingDirectory);
+		if (Output.bSucceeded)
+		{
+			Output.GenerateOutputHash();
+			if (Input.CompressionFormat != NAME_None)
+			{
+				Output.CompressOutput(Input.CompressionFormat, Input.OodleCompressor, Input.OodleLevel);
+			}
+		}
+		Output.CompileTime = FPlatformTime::Seconds() - TimeStart;
+
+		if (Compiler->UsesHLSLcc(Input))
+		{
+			Output.bUsedHLSLccCompiler = true;
+		}
+
+		if (CompileCount)
+		{
+			++(*CompileCount);
+		}
+	}
+#if PLATFORM_WINDOWS
+	__except (HandleShaderCompileException(GetExceptionInformation(), OutExceptionMsg, OutExceptionCallstack))
+	{
+		Output.bSucceeded = false;
+	}
+#endif
+}
+
+void CompileShader(const TArray<const IShaderFormat*>& ShaderFormats, FShaderCompilerInput& Input, FShaderCompilerOutput& Output, const FString& WorkingDirectory, int32* CompileCount)
+{
+	TRACE_CPUPROFILER_EVENT_SCOPE(CompileShader);
+
+	const IShaderFormat* Compiler = FindShaderFormat(Input.ShaderFormat, ShaderFormats);
+	if (!Compiler)
+	{
+		UE_LOG(LogShaders, Fatal, TEXT("Can't compile shaders for format %s, couldn't load compiler dll"), *Input.ShaderFormat.ToString());
+	}
+
+	if (IsValidRef(Input.SharedEnvironment))
+	{
+		Input.Environment.Merge(*Input.SharedEnvironment);
+	}
+
+#if PLATFORM_WINDOWS
+	FString ExceptionMsg;
+	FString ExceptionCallstack;
+	CompileShaderInternal(Compiler, Input, Output, WorkingDirectory, ExceptionMsg, ExceptionCallstack, CompileCount);
+	if (!Output.bSucceeded && (!ExceptionMsg.IsEmpty() || !ExceptionCallstack.IsEmpty()))
+	{
+		FShaderCompilerError Error;
+		Error.StrippedErrorMessage = FString::Printf(
+			TEXT("Exception encountered in platform compiler: %s\nException Callstack:\n%s"), 
+			*ExceptionMsg, 
+			*ExceptionCallstack);
+		Output.Errors.Add(Error);
+	}
+#else
+	CompileShaderInternal(Compiler, Input, Output, WorkingDirectory, CompileCount);
+#endif
+}
+
 /**
 * Add a new entry to the list of shader source files
 * Only unique entries which can be loaded are added as well as their #include files
diff --git a/Engine/Source/Runtime/RenderCore/Public/ShaderCompilerCore.h b/Engine/Source/Runtime/RenderCore/Public/ShaderCompilerCore.h
--- a/Engine/Source/Runtime/RenderCore/Public/ShaderCompilerCore.h
+++ b/Engine/Source/Runtime/RenderCore/Public/ShaderCompilerCore.h
@@ -16,6 +16,7 @@
 #include "ShaderCore.h"
 
 class Error;
+class IShaderFormat;
 
 // this is for the protocol, not the data, bump if FShaderCompilerInput or ProcessInputFromArchive changes.
 inline const int32 ShaderCompileWorkerInputVersion = 18;
@@ -534,6 +535,12 @@ struct RENDERCORE_API FSCWErrorCode
 	static bool IsSet();
 };
 
+#if PLATFORM_WINDOWS
+extern RENDERCORE_API int HandleShaderCompileException(Windows::LPEXCEPTION_POINTERS Info, FString& OutExMsg, FString& OutCallStack);
+#endif
+extern RENDERCORE_API const IShaderFormat* FindShaderFormat(FName Format, TArray<const IShaderFormat*> ShaderFormats);
+extern RENDERCORE_API void CompileShader(const TArray<const IShaderFormat*>& ShaderFormats, FShaderCompilerInput& Input, FShaderCompilerOutput& Output, const FString& WorkingDirectory, int32* CompileCount = nullptr);
+
 UE_DEPRECATED(5.2, "Functionality has moved to UE::ShaderCompilerCommon::ShouldUseStableConstantBuffer")
 inline bool ShouldUseStableConstantBuffer(const FShaderCompilerInput& Input)
 {
