ShaderCompileWorker: Fix postbuild copy step if the target name is renamed for any reason

#rnx

diff --git a/Engine/Source/Programs/ShaderCompileWorker/ShaderCompileWorker.Target.cs b/Engine/Source/Programs/ShaderCompileWorker/ShaderCompileWorker.Target.cs
--- a/Engine/Source/Programs/ShaderCompileWorker/ShaderCompileWorker.Target.cs
+++ b/Engine/Source/Programs/ShaderCompileWorker/ShaderCompileWorker.Target.cs
@@ -19,7 +19,7 @@ public class ShaderCompileWorkerTarget : TargetRules
         {
             // The interception interface in XGE requires that the parent and child processes have different filenames on disk.
             // To avoid building an entire separate worker just for this, we duplicate the ShaderCompileWorker in a post build step.
-            const string SrcPath  = "$(EngineDir)\\Binaries\\$(TargetPlatform)\\ShaderCompileWorker.exe";
+            const string SrcPath  = "$(EngineDir)\\Binaries\\$(TargetPlatform)\\$(TargetName).exe";
             const string DestPath = "$(EngineDir)\\Binaries\\$(TargetPlatform)\\XGEControlWorker.exe";
 
             PostBuildSteps.Add(string.Format("echo Copying {0} to {1}", SrcPath, DestPath));
