Add "-DebugInfoPath=PATH" argument to SCW to allow dumping debug info with -directcompile commandlet.

#rnx

diff --git a/Engine/Source/Programs/ShaderCompileWorker/Private/ShaderCompileWorker.cpp b/Engine/Source/Programs/ShaderCompileWorker/Private/ShaderCompileWorker.cpp
--- a/Engine/Source/Programs/ShaderCompileWorker/Private/ShaderCompileWorker.cpp
+++ b/Engine/Source/Programs/ShaderCompileWorker/Private/ShaderCompileWorker.cpp
@@ -853,6 +853,7 @@ static void DirectCompile(const TArray<const class IShaderFormat*>& ShaderFormat
 	FName FormatName;
 	FName ShaderPlatformName;
 	FString Entry = TEXT("Main");
+	FString DumpDebugInfoPath;
 	bool bPipeline = false;
 	EShaderFrequency Frequency = SF_Pixel;
 	TArray<FString> UsedOutputs;
@@ -884,6 +885,10 @@ static void DirectCompile(const TArray<const class IShaderFormat*>& ShaderFormat
 			{
 				CFlags = FCString::Atoi64(*Token.RightChop(7));
 			}
+			else if (Token.StartsWith(TEXT("DebugInfoPath=")))
+			{
+				DumpDebugInfoPath = Token.RightChop(14);
+			}
 			else if (!FCString::Strcmp(*Token, TEXT("ps")))
 			{
 				Frequency = SF_Pixel;
@@ -962,9 +967,8 @@ static void DirectCompile(const TArray<const class IShaderFormat*>& ShaderFormat
 	Job.Input.Target.Platform =  ShaderFormatNameToShaderPlatform(FormatName);
 	Job.Input.Target.Frequency = Frequency;
 	Job.Input.bSkipPreprocessedCache = true;
-
+	Job.Input.DumpDebugInfoPath = DumpDebugInfoPath;
 	Job.Input.Environment.CompilerFlags = FShaderCompilerFlags(CFlags);
-
 	Job.Input.bCompilingForShaderPipeline = bPipeline;
 	Job.Input.bIncludeUsedOutputs = bIncludeUsedOutputs;
 	Job.Input.UsedOutputs = UsedOutputs;
