Adding -NumJobs=N argument to SCW to facilitate in profiling. This argument causes the compiler to run the processing loop at most N times.

diff --git a/Engine/Source/Programs/ShaderCompileWorker/Private/ShaderCompileWorker.cpp b/Engine/Source/Programs/ShaderCompileWorker/Private/ShaderCompileWorker.cpp
--- a/Engine/Source/Programs/ShaderCompileWorker/Private/ShaderCompileWorker.cpp
+++ b/Engine/Source/Programs/ShaderCompileWorker/Private/ShaderCompileWorker.cpp
@@ -187,6 +187,7 @@ class FWorkLoop
 public:
 	// If we have been idle for 20 seconds then exit. Can be overriden from the cmd line with -TimeToLive=N where N is in seconds (and a float value)
 	float TimeToLive = 20.0f;
+	int32 NumberToProcess = -1;
 	bool DisableFileWrite = false;
 	bool KeepInput = false;
 
@@ -215,6 +216,10 @@ public:
 			{
 				KeepInput = true;
 			}
+			else if (Switch.StartsWith(TEXT("NumJobs=")))
+			{
+				NumberToProcess = FCString::Atoi(Switch.GetCharArray().GetData() + 8);
+			}
 		}
 	}
 
@@ -223,6 +228,8 @@ public:
 		UE_LOG(LogShaders, Log, TEXT("Entering job loop"));
 		TRACE_CPUPROFILER_EVENT_SCOPE(Loop);
 
+		int32 NumberProcessed = 0;
+
 		while(true)
 		{
 			TArray<FShaderCompileJob> SingleJobs;
@@ -284,6 +291,13 @@ public:
 				break;
 			}
 
+			NumberProcessed++;
+			if (NumberToProcess > 0 && NumberProcessed > NumberToProcess)
+			{
+				UE_LOG(LogShaders, Log, TEXT("NumJobs limit hit"));
+				break;
+			}
+
 #if ENABLE_LOW_LEVEL_MEM_TRACKER
 			{
 				TRACE_CPUPROFILER_EVENT_SCOPE(UpdateStatsPerFrame);
