[ShaderCompileWorker]
* Added support for process reuse when running through uba

diff --git a/Engine/Source/Programs/ShaderCompileWorker/Private/ShaderCompileWorker.cpp b/Engine/Source/Programs/ShaderCompileWorker/Private/ShaderCompileWorker.cpp
--- a/Engine/Source/Programs/ShaderCompileWorker/Private/ShaderCompileWorker.cpp
+++ b/Engine/Source/Programs/ShaderCompileWorker/Private/ShaderCompileWorker.cpp
@@ -285,6 +285,32 @@ public:
 				break;
 			}
 
+#if PLATFORM_WINDOWS // Currently only implemented for windows
+			static HMODULE UbaDetoursModule = GetModuleHandleW(L"UbaDetours.dll");
+			if (UbaDetoursModule)
+			{
+				using UbaRequestNextProcessFunc = bool(TCHAR* outArguments, uint32 outArgumentsCapacity);
+				static UbaRequestNextProcessFunc* RequestNextProcess = (UbaRequestNextProcessFunc*)(void*)GetProcAddress(UbaDetoursModule, "UbaRequestNextProcess");
+
+				// Request new process
+				TCHAR Arguments[1024];
+				if (RequestNextProcess(Arguments, 1024))
+				{
+					TArray<FString> Tokens;
+					TArray<FString> Switches;
+					FCommandLine::Parse(Arguments, Tokens, Switches);
+
+					InputFilename = Tokens[3];
+					OutputFilename = Tokens[4];
+
+					InputFilePath = WorkingDirectory / InputFilename;
+					OutputFilePath = WorkingDirectory / OutputFilename;
+
+					continue;
+				}
+			}
+#endif
+
 			if (TimeToLive == 0)
 			{
 				UE_LOG(LogShaders, Log, TEXT("TimeToLive set to 0, exiting after single job"));
@@ -313,11 +339,11 @@ public:
 private:
 	const int32 ParentProcessId;
 	const FString WorkingDirectory;
-	const FString InputFilename;
-	const FString OutputFilename;
+	FString InputFilename;
+	FString OutputFilename;
 
-	const FString InputFilePath;
-	const FString OutputFilePath;
+	FString InputFilePath;
+	FString OutputFilePath;
 	TMap<FString, uint32> FormatVersionMap;
 	FString TempFilePath;
 
