[ShaderCompileWorker]
* Changed so when building through uba we don't delete the .in file. Since uba supports rescheduling when remote helpers disconnect we can't delete this file until out is fully sent back

diff --git a/Engine/Source/Programs/ShaderCompileWorker/Private/ShaderCompileWorker.cpp b/Engine/Source/Programs/ShaderCompileWorker/Private/ShaderCompileWorker.cpp
--- a/Engine/Source/Programs/ShaderCompileWorker/Private/ShaderCompileWorker.cpp
+++ b/Engine/Source/Programs/ShaderCompileWorker/Private/ShaderCompileWorker.cpp
@@ -50,6 +50,25 @@ static void OnXGEJobCompleted(const TCHAR* WorkingDirectory)
 	}
 }
 
+#if PLATFORM_WINDOWS // Currently only implemented for windows
+HMODULE GetUbaModule()
+{
+	static HMODULE UbaDetoursModule = GetModuleHandleW(L"UbaDetours.dll");
+	return UbaDetoursModule;
+}
+#endif
+
+bool IsRunningRemoteUba()
+{
+#if PLATFORM_WINDOWS // Currently only implemented for windows
+	using UbaRunningRemoteFunc = bool();
+	static UbaRunningRemoteFunc* RequestNextProcess = (UbaRunningRemoteFunc*)(void*)GetProcAddress(GetUbaModule(), "UbaRunningRemote");
+	return RequestNextProcess ? RequestNextProcess() : false;
+#else
+	return false;
+#endif
+}
+
 #if USING_CODE_ANALYSIS
 	UE_NORETURN static inline void ExitWithoutCrash(FSCWErrorCode::ECode ErrorCode, const FString& Message);
 #endif
@@ -286,8 +305,7 @@ public:
 			}
 
 #if PLATFORM_WINDOWS // Currently only implemented for windows
-			static HMODULE UbaDetoursModule = GetModuleHandleW(L"UbaDetours.dll");
-			if (UbaDetoursModule)
+			if (HMODULE UbaDetoursModule = GetUbaModule())
 			{
 				using UbaRequestNextProcessFunc = bool(TCHAR* outArguments, uint32 outArgumentsCapacity);
 				static UbaRequestNextProcessFunc* RequestNextProcess = (UbaRequestNextProcessFunc*)(void*)GetProcAddress(UbaDetoursModule, "UbaRequestNextProcess");
@@ -718,7 +736,8 @@ private:
 		// Don't delete the input file if we are running under Incredibuild (or if the cmdline args explicitly told us to keep it).
 		// In xml mode, we signal completion by creating a zero byte "Success" file after the output file has been fully written.
 		// In intercept mode, completion is signaled by this process terminating.
-		if (!IsUsingXGE() && !KeepInput)
+		// For UBA we can't delete the file when running remotely because there might be a crash or disconnect happening before result is sent back and then we can't retry
+		if (!IsUsingXGE() && !KeepInput && !IsRunningRemoteUba())
 		{
 			do 
 			{
