[ShaderCompileWorker]
* Fixed so WorkingDirectory is updated when new process is fetched

diff --git a/Engine/Source/Programs/ShaderCompileWorker/Private/ShaderCompileWorker.cpp b/Engine/Source/Programs/ShaderCompileWorker/Private/ShaderCompileWorker.cpp
--- a/Engine/Source/Programs/ShaderCompileWorker/Private/ShaderCompileWorker.cpp
+++ b/Engine/Source/Programs/ShaderCompileWorker/Private/ShaderCompileWorker.cpp
@@ -62,8 +62,8 @@ bool IsRunningRemoteUba()
 {
 #if PLATFORM_WINDOWS // Currently only implemented for windows
 	using UbaRunningRemoteFunc = bool();
-	static UbaRunningRemoteFunc* RequestNextProcess = (UbaRunningRemoteFunc*)(void*)GetProcAddress(GetUbaModule(), "UbaRunningRemote");
-	return RequestNextProcess ? RequestNextProcess() : false;
+	static UbaRunningRemoteFunc* RunningRemote = (UbaRunningRemoteFunc*)(void*)GetProcAddress(GetUbaModule(), "UbaRunningRemote");
+	return RunningRemote ? RunningRemote() : false;
 #else
 	return false;
 #endif
@@ -318,6 +318,7 @@ public:
 					TArray<FString> Switches;
 					FCommandLine::Parse(Arguments, Tokens, Switches);
 
+					WorkingDirectory = Tokens[0];
 					InputFilename = Tokens[3];
 					OutputFilename = Tokens[4];
 
@@ -356,7 +357,7 @@ public:
 
 private:
 	const int32 ParentProcessId;
-	const FString WorkingDirectory;
+	FString WorkingDirectory;
 	FString InputFilename;
 	FString OutputFilename;
 
