[ShaderCompileWorker]
* Removed code checking if worker is running remotely. local and remote should be treated the same way

diff --git a/Engine/Source/Programs/ShaderCompileWorker/Private/ShaderCompileWorker.cpp b/Engine/Source/Programs/ShaderCompileWorker/Private/ShaderCompileWorker.cpp
--- a/Engine/Source/Programs/ShaderCompileWorker/Private/ShaderCompileWorker.cpp
+++ b/Engine/Source/Programs/ShaderCompileWorker/Private/ShaderCompileWorker.cpp
@@ -58,17 +58,6 @@ HMODULE GetUbaModule()
 }
 #endif
 
-bool IsRunningRemoteUba()
-{
-#if PLATFORM_WINDOWS // Currently only implemented for windows
-	using UbaRunningRemoteFunc = bool();
-	static UbaRunningRemoteFunc* RunningRemote = (UbaRunningRemoteFunc*)(void*)GetProcAddress(GetUbaModule(), "UbaRunningRemote");
-	return RunningRemote ? RunningRemote() : false;
-#else
-	return false;
-#endif
-}
-
 #if USING_CODE_ANALYSIS
 	UE_NORETURN static inline void ExitWithoutCrash(FSCWErrorCode::ECode ErrorCode, const FString& Message);
 #endif
@@ -738,7 +727,7 @@ private:
 		// In xml mode, we signal completion by creating a zero byte "Success" file after the output file has been fully written.
 		// In intercept mode, completion is signaled by this process terminating.
 		// For UBA we can't delete the file when running remotely because there might be a crash or disconnect happening before result is sent back and then we can't retry
-		if (!IsUsingXGE() && !KeepInput && !IsRunningRemoteUba())
+		if (!IsUsingXGE() && !KeepInput && !GetUbaModule())
 		{
 			do 
 			{
