[ShaderCompileWorker]
* Fixed linux compile error

diff --git a/Engine/Source/Programs/ShaderCompileWorker/Private/ShaderCompileWorker.cpp b/Engine/Source/Programs/ShaderCompileWorker/Private/ShaderCompileWorker.cpp
--- a/Engine/Source/Programs/ShaderCompileWorker/Private/ShaderCompileWorker.cpp
+++ b/Engine/Source/Programs/ShaderCompileWorker/Private/ShaderCompileWorker.cpp
@@ -58,6 +58,15 @@ HMODULE GetUbaModule()
 }
 #endif
 
+inline bool IsUsingUBA()
+{
+#if PLATFORM_WINDOWS // Currently only implemented for windows
+	return GetUbaModule() != nullptr;
+#else
+	return false;
+#endif
+}
+
 #if USING_CODE_ANALYSIS
 	UE_NORETURN static inline void ExitWithoutCrash(FSCWErrorCode::ECode ErrorCode, const FString& Message);
 #endif
@@ -727,7 +736,7 @@ private:
 		// In xml mode, we signal completion by creating a zero byte "Success" file after the output file has been fully written.
 		// In intercept mode, completion is signaled by this process terminating.
 		// For UBA we can't delete the file when running remotely because there might be a crash or disconnect happening before result is sent back and then we can't retry
-		if (!IsUsingXGE() && !KeepInput && !GetUbaModule())
+		if (!IsUsingXGE() && !KeepInput && !IsUsingUBA())
 		{
 			do 
 			{
