[ShaderCompiler]
* Fixed function pointer format for getting next process in uba
* Fixed so output files are written even when return code is non-zero (error)

diff --git a/Engine/Plugins/UbaController/Source/UbaController/Private/UbaJobProcessor.cpp b/Engine/Plugins/UbaController/Source/UbaController/Private/UbaJobProcessor.cpp
--- a/Engine/Plugins/UbaController/Source/UbaController/Private/UbaJobProcessor.cpp
+++ b/Engine/Plugins/UbaController/Source/UbaController/Private/UbaJobProcessor.cpp
@@ -231,6 +231,7 @@ void FUbaJobProcessor::RunTaskWithUba(FTask* Task)
 	ProcessInfo.arguments = *Parameters;
 	ProcessInfo.description = *InputFileName;
 	ProcessInfo.workingDir = *AppDir;
+	ProcessInfo.writeOutputFilesOnFail = true;
 
 	if (UbaJobProcessorOptions::bProcessLogEnabled)
 		ProcessInfo.logFile = *LogPath;
diff --git a/Engine/Source/Programs/ShaderCompileWorker/Private/ShaderCompileWorker.cpp b/Engine/Source/Programs/ShaderCompileWorker/Private/ShaderCompileWorker.cpp
--- a/Engine/Source/Programs/ShaderCompileWorker/Private/ShaderCompileWorker.cpp
+++ b/Engine/Source/Programs/ShaderCompileWorker/Private/ShaderCompileWorker.cpp
@@ -305,12 +305,12 @@ public:
 #if PLATFORM_WINDOWS // Currently only implemented for windows
 			if (HMODULE UbaDetoursModule = GetUbaModule())
 			{
-				using UbaRequestNextProcessFunc = bool(TCHAR* outArguments, uint32 outArgumentsCapacity);
+				using UbaRequestNextProcessFunc = bool(uint32 prevExitCode, TCHAR* outArguments, uint32 outArgumentsCapacity);
 				static UbaRequestNextProcessFunc* RequestNextProcess = (UbaRequestNextProcessFunc*)(void*)GetProcAddress(UbaDetoursModule, "UbaRequestNextProcess");
 
 				// Request new process
 				TCHAR Arguments[1024];
-				if (RequestNextProcess(Arguments, 1024))
+				if (RequestNextProcess(0, Arguments, 1024))
 				{
 					TArray<FString> Tokens;
 					TArray<FString> Switches;
