[ShaderWorker]
* Changed so shader worker breaks out of loop when in uba path and no new processes are available

diff --git a/Engine/Source/Programs/ShaderCompileWorker/Private/ShaderCompileWorker.cpp b/Engine/Source/Programs/ShaderCompileWorker/Private/ShaderCompileWorker.cpp
--- a/Engine/Source/Programs/ShaderCompileWorker/Private/ShaderCompileWorker.cpp
+++ b/Engine/Source/Programs/ShaderCompileWorker/Private/ShaderCompileWorker.cpp
@@ -310,21 +310,22 @@ public:
 
 				// Request new process
 				TCHAR Arguments[1024];
-				if (RequestNextProcess(0, Arguments, 1024))
+				if (!RequestNextProcess(0, Arguments, 1024))
 				{
-					TArray<FString> Tokens;
-					TArray<FString> Switches;
-					FCommandLine::Parse(Arguments, Tokens, Switches);
+					break; // No process available, exit loop
+				}
 
-					WorkingDirectory = Tokens[0];
-					InputFilename = Tokens[3];
-					OutputFilename = Tokens[4];
+				TArray<FString> Tokens;
+				TArray<FString> Switches;
+				FCommandLine::Parse(Arguments, Tokens, Switches);
 
-					InputFilePath = WorkingDirectory / InputFilename;
-					OutputFilePath = WorkingDirectory / OutputFilename;
+				WorkingDirectory = Tokens[0];
+				InputFilename = Tokens[3];
+				OutputFilename = Tokens[4];
 
-					continue;
-				}
+				InputFilePath = WorkingDirectory / InputFilename;
+				OutputFilePath = WorkingDirectory / OutputFilename;
+				continue;
 			}
 #endif
 
