[ShaderCompileWorker]
* Fixed so CrashOutputFile is updated when process is reused (the crash output file should be the same as the output file)

diff --git a/Engine/Source/Programs/ShaderCompileWorker/Private/ShaderCompileWorker.cpp b/Engine/Source/Programs/ShaderCompileWorker/Private/ShaderCompileWorker.cpp
--- a/Engine/Source/Programs/ShaderCompileWorker/Private/ShaderCompileWorker.cpp
+++ b/Engine/Source/Programs/ShaderCompileWorker/Private/ShaderCompileWorker.cpp
@@ -240,7 +240,7 @@ public:
 		}
 	}
 
-	void Loop()
+	void Loop(FString& CrashOutputFile)
 	{
 		UE_LOG(LogShaders, Log, TEXT("Entering job loop"));
 		TRACE_CPUPROFILER_EVENT_SCOPE(Loop);
@@ -315,6 +315,8 @@ public:
 					break; // No process available, exit loop
 				}
 
+				// We got a new process, change inputs and outputs and run again
+				
 				TArray<FString> Tokens;
 				TArray<FString> Switches;
 				FCommandLine::Parse(Arguments, Tokens, Switches);
@@ -325,6 +327,8 @@ public:
 
 				InputFilePath = WorkingDirectory / InputFilename;
 				OutputFilePath = WorkingDirectory / OutputFilename;
+
+				CrashOutputFile = OutputFilePath;
 				continue;
 			}
 #endif
@@ -1040,7 +1044,7 @@ static void DirectCompile(const TArray<const class IShaderFormat*>& ShaderFormat
  *		The parent process Id
  *		The thread Id corresponding to this worker
  */
-static int32 GuardedMain(int32 argc, TCHAR* argv[], bool bDirectMode)
+static int32 GuardedMain(int32 argc, TCHAR* argv[], FString& CrashOutputFile, bool bDirectMode)
 {
 	FString ExtraCmdLine = TEXT("-NOPACKAGECACHE -ReduceThreadUsage -cpuprofilertrace -nocrashreports");
 
@@ -1178,14 +1182,14 @@ static int32 GuardedMain(int32 argc, TCHAR* argv[], bool bDirectMode)
 		TRACE_CPUPROFILER_EVENT_SCOPE(FWorkLoop);
 
 		FWorkLoop WorkLoop(argv[2], argv[1], argv[4], argv[5], FormatVersionMap);
-		WorkLoop.Loop();
+		WorkLoop.Loop(CrashOutputFile);
 	}
 
 	return 0;
 }
 
 
-static int32 GuardedMainWrapper(int32 ArgC, TCHAR* ArgV[], const TCHAR* CrashOutputFile, bool bDirectMode)
+static int32 GuardedMainWrapper(int32 ArgC, TCHAR* ArgV[], FString& CrashOutputFile, bool bDirectMode)
 {
 	FTaskTagScope Scope(ETaskTag::EGameThread);
 	// We need to know whether we are using XGE now, in case an exception
@@ -1208,7 +1212,7 @@ static int32 GuardedMainWrapper(int32 ArgC, TCHAR* ArgV[], const TCHAR* CrashOut
 	if (FPlatformMisc::IsDebuggerPresent())
 #endif
 	{
-		ReturnCode = GuardedMain(ArgC, ArgV, bDirectMode);
+		ReturnCode = GuardedMain(ArgC, ArgV, CrashOutputFile, bDirectMode);
 	}
 #if PLATFORM_WINDOWS
 	else
@@ -1220,12 +1224,12 @@ static int32 GuardedMainWrapper(int32 ArgC, TCHAR* ArgV[], const TCHAR* CrashOut
 		__try
 		{
 			GIsGuarded = 1;
-			ReturnCode = GuardedMain(ArgC, ArgV, bDirectMode);
+			ReturnCode = GuardedMain(ArgC, ArgV, CrashOutputFile, bDirectMode);
 			GIsGuarded = 0;
 		}
 		__except(HandleShaderCompileException(GetExceptionInformation(), ExceptionMsg, ExceptionCallStack))
 		{
-			FArchive& OutputFile = *IFileManager::Get().CreateFileWriter(CrashOutputFile, FILEWRITE_EvenIfReadOnly);
+			FArchive& OutputFile = *IFileManager::Get().CreateFileWriter(*CrashOutputFile, FILEWRITE_EvenIfReadOnly);
 
 			if (GFailedErrorCode == FSCWErrorCode::Success)
 			{
@@ -1258,6 +1262,10 @@ static int32 GuardedMainWrapper(int32 ArgC, TCHAR* ArgV[], const TCHAR* CrashOut
 				ReturnCode = 1;
 				OnXGEJobCompleted(ArgV[1]);
 			}
+			else if (GetUbaModule())
+			{
+				ReturnCode = GFailedErrorCode;
+			}
 		}
 	}
 #endif
@@ -1319,5 +1327,5 @@ INT32_MAIN_INT32_ARGC_TCHAR_ARGV()
 		OutputFilePath += ArgV[5];
 	}
 
-	return GuardedMainWrapper(ArgC, ArgV, *OutputFilePath, bDirectMode);
+	return GuardedMainWrapper(ArgC, ArgV, OutputFilePath, bDirectMode);
 }
