[SCW] Provide FilePath in assertion when SCW fails check that parent process is still running.

This will provide the detail needed to connect failed compile jobs (due terminated parent processes) to their UBA job:
Previously, such error looked like this:
> LogUbaController: Display: [CookWorker 4]: LogWindows: Error: appError called: Assertion failed: IFileManager::Get().FileSize(*FilePath) == INDEX_NONE [File:[...]\ShaderCompileWorker.cpp] [Line: 877]
> LogUbaController: Display: [CookWorker 4]: Exiting due to OpenProcess(ParentProcessId) failing and the input file is present!

And now provides the file path:
> LogUbaController: Display: [CookWorker 4]: LogWindows: Error: appError called: Assertion failed: IFileManager::Get().FileSize(*FilePath) == INDEX_NONE [File:[...]\ShaderCompileWorker.cpp] [Line: 877]
> LogUbaController: Display: [CookWorker 4]: Exiting due to OpenProcess(ParentProcessId) failing and the input file is present! FilePath=D:\Example\42.uba.out

#rnx

diff --git a/Engine/Source/Programs/ShaderCompileWorker/Private/ShaderCompileWorker.cpp b/Engine/Source/Programs/ShaderCompileWorker/Private/ShaderCompileWorker.cpp
--- a/Engine/Source/Programs/ShaderCompileWorker/Private/ShaderCompileWorker.cpp
+++ b/Engine/Source/Programs/ShaderCompileWorker/Private/ShaderCompileWorker.cpp
@@ -848,7 +848,7 @@ private:
 			if (!FPlatformProcess::IsApplicationRunning(ParentProcessId))
 			{
 				FString FilePath = FString(WorkingDirectory) + InputFilename;
-				checkf(IFileManager::Get().FileSize(*FilePath) == INDEX_NONE, TEXT("Exiting due to the parent process no longer running and the input file is present!"));
+				checkf(IFileManager::Get().FileSize(*FilePath) == INDEX_NONE, TEXT("Exiting due to the parent process no longer running! FilePath=%s"), *FilePath);
 				UE_LOG(LogShaders, Log, TEXT("Parent process no longer running, exiting"));
 				FPlatformMisc::RequestExit(false);
 			}
@@ -874,7 +874,7 @@ private:
 				// If we couldn't open the process then it is no longer running, exit
 				if (ParentProcessHandle == nullptr)
 				{
-					checkf(IFileManager::Get().FileSize(*FilePath) == INDEX_NONE, TEXT("Exiting due to OpenProcess(ParentProcessId) failing and the input file is present!"));
+					checkf(IFileManager::Get().FileSize(*FilePath) == INDEX_NONE, TEXT("Exiting due to OpenProcess(ParentProcessId) failing! FilePath=%s"), *FilePath);
 					UE_LOG(LogShaders, Log, TEXT("Couldn't OpenProcess, Parent process no longer running, exiting"));
 					FPlatformMisc::RequestExit(false);
 				}
@@ -886,7 +886,7 @@ private:
 					uint32 WaitResult = WaitForSingleObject(ParentProcessHandle, 0);
 					if (WaitResult != WAIT_TIMEOUT)
 					{
-						checkf(IFileManager::Get().FileSize(*FilePath) == INDEX_NONE, TEXT("Exiting due to WaitForSingleObject(ParentProcessHandle) signaling and the input file is present!"));
+						checkf(IFileManager::Get().FileSize(*FilePath) == INDEX_NONE, TEXT("Exiting due to WaitForSingleObject(ParentProcessHandle) signaling! FilePath=%s"), *FilePath);
 						UE_LOG(LogShaders, Log, TEXT("WaitForSingleObject signaled, Parent process no longer running, exiting"));
 						FPlatformMisc::RequestExit(false);
 					}
