Fix CIS build warning by using bUseLoggingInShipping instead of using USE_LOGGING_IN_SHIPPING directly.

#rnx

diff --git a/Engine/Source/Programs/ShaderCompileWorker/ShaderCompileWorker.Target.cs b/Engine/Source/Programs/ShaderCompileWorker/ShaderCompileWorker.Target.cs
--- a/Engine/Source/Programs/ShaderCompileWorker/ShaderCompileWorker.Target.cs
+++ b/Engine/Source/Programs/ShaderCompileWorker/ShaderCompileWorker.Target.cs
@@ -55,7 +55,7 @@ public class ShaderCompileWorkerTarget : TargetRules
 		// Allow logging everything to memory so we can decide to dump everything to file when SCW crashed.
 		// This still requires the commandline argument -LogToMemory, so the shader compiling manager can decide to enable/disable it at runtime.
 		bLoggingToMemoryEnabled = true;
-		GlobalDefinitions.Add("USE_LOGGING_IN_SHIPPING=1");
+		bUseLoggingInShipping = true;
 
 		// Linking against wer.lib/wer.dll causes XGE to bail when the worker is run on a Windows 8 machine, so turn this off.
 		GlobalDefinitions.Add("ALLOW_WINDOWS_ERROR_REPORT_LIB=0");
