[Shaders] Put SCW into critical error mode during exception handling to dump log from memory to disk. Appendix to CL 35144434.

GIsCriticalError was not set to true inside exception handling block, preventing logs in SCW to be dumped from memory to disk despite -LogToMemory argument.

#rnx

diff --git a/Engine/Source/Programs/ShaderCompileWorker/Private/ShaderCompileWorker.cpp b/Engine/Source/Programs/ShaderCompileWorker/Private/ShaderCompileWorker.cpp
--- a/Engine/Source/Programs/ShaderCompileWorker/Private/ShaderCompileWorker.cpp
+++ b/Engine/Source/Programs/ShaderCompileWorker/Private/ShaderCompileWorker.cpp
@@ -1206,6 +1206,9 @@ static int32 GuardedMainWrapper(int32 ArgC, TCHAR* ArgV[], FString& CrashOutputF
 		}
 		__except(HandleShaderCompileException(GetExceptionInformation(), ExceptionMsg, ExceptionCallStack))
 		{
+			// Put app into critical error mode to allow dumping logs from memory to disk
+			GIsCriticalError = true;
+
 			FArchive& OutputFile = *IFileManager::Get().CreateFileWriter(*CrashOutputFile, FILEWRITE_EvenIfReadOnly);
 
 			if (GFailedErrorCode == FSCWErrorCode::Success)
