[ShaderCompileWorker]
* Added so IsRunningUba works on macos

diff --git a/Engine/Source/Programs/ShaderCompileWorker/Private/ShaderCompileWorker.cpp b/Engine/Source/Programs/ShaderCompileWorker/Private/ShaderCompileWorker.cpp
--- a/Engine/Source/Programs/ShaderCompileWorker/Private/ShaderCompileWorker.cpp
+++ b/Engine/Source/Programs/ShaderCompileWorker/Private/ShaderCompileWorker.cpp
@@ -22,6 +22,10 @@
 #include "Serialization/MemoryReader.h"
 #include "SocketSubsystem.h"
 
+#if PLATFORM_MAC
+#include <mach-o/dyld.h>
+#endif
+
 #define DEBUG_USING_CONSOLE	0
 
 static double LastCompileTime = 0.0;
@@ -65,9 +69,16 @@ inline bool IsUsingUBA()
 {
 #if PLATFORM_WINDOWS // Currently only implemented for windows
 	return GetUbaModule() != nullptr;
-#else
+#elif PLATFORM_MAC
+		for (int i = 0, e = _dyld_image_count(); i != e; i++)
+		{
+			if (strstr(_dyld_get_image_name(i), "UbaDetours.dylib"))
+			{
+				return true;
+			}
+		}
+#endif	
 	return false;
-#endif
 }
 
 #if USING_CODE_ANALYSIS
@@ -334,6 +345,11 @@ public:
 				CrashOutputFile = OutputFilePath;
 				continue;
 			}
+#else
+			if (IsUsingUBA())
+			{
+				break;
+			}
 #endif
 
 			if (TimeToLive == 0)
