[ShaderCompileWorker]
* Fixed farm compile error

diff --git a/Engine/Source/Programs/ShaderCompileWorker/Private/ShaderCompileWorker.cpp b/Engine/Source/Programs/ShaderCompileWorker/Private/ShaderCompileWorker.cpp
--- a/Engine/Source/Programs/ShaderCompileWorker/Private/ShaderCompileWorker.cpp
+++ b/Engine/Source/Programs/ShaderCompileWorker/Private/ShaderCompileWorker.cpp
@@ -70,15 +70,17 @@ inline bool IsUsingUBA()
 #if PLATFORM_WINDOWS // Currently only implemented for windows
 	return GetUbaModule() != nullptr;
 #elif PLATFORM_MAC
-		for (int i = 0, e = _dyld_image_count(); i != e; i++)
+	for (int i = 0, e = _dyld_image_count(); i != e; i++)
+	{
+		if (strstr(_dyld_get_image_name(i), "UbaDetours.dylib"))
 		{
-			if (strstr(_dyld_get_image_name(i), "UbaDetours.dylib"))
-			{
-				return true;
-			}
+			return true;
 		}
-#endif	
+	}
+	return false;
+#else
 	return false;
+#endif	
 }
 
 #if USING_CODE_ANALYSIS
