Unshelved from pending changelist '36032105':

[ShaderCompileWorker]
* Added support for process reuse on macos

diff --git a/Engine/Source/Programs/ShaderCompileWorker/Private/ShaderCompileWorker.cpp b/Engine/Source/Programs/ShaderCompileWorker/Private/ShaderCompileWorker.cpp
--- a/Engine/Source/Programs/ShaderCompileWorker/Private/ShaderCompileWorker.cpp
+++ b/Engine/Source/Programs/ShaderCompileWorker/Private/ShaderCompileWorker.cpp
@@ -318,24 +318,60 @@ public:
 				break;
 			}
 
-#if PLATFORM_WINDOWS // Currently only implemented for windows
-			if (HMODULE UbaDetoursModule = GetUbaModule())
+			if (IsUsingUBA())
 			{
-				using UbaRequestNextProcessFunc = bool(uint32 prevExitCode, TCHAR* outArguments, uint32 outArgumentsCapacity);
-				static UbaRequestNextProcessFunc* RequestNextProcess = (UbaRequestNextProcessFunc*)(void*)GetProcAddress(UbaDetoursModule, "UbaRequestNextProcess");
+#if PLATFORM_WINDOWS
+				using ubachar = TCHAR;
+#else
+				using ubachar = char;
+#endif
+				
+				using UbaRequestNextProcessFunc = bool(uint32 prevExitCode, ubachar* outArguments, uint32 outArgumentsCapacity);
+				static UbaRequestNextProcessFunc* RequestNextProcess = nullptr;
+				if (!RequestNextProcess)
+				{
+#if PLATFORM_WINDOWS
+					if (HMODULE UbaDetoursModule = GetUbaModule())
+					{
+						RequestNextProcess = (UbaRequestNextProcessFunc*)(void*)GetProcAddress(UbaDetoursModule, "UbaRequestNextProcess");
+					}
+#elif PLATFORM_MAC
+					if (void* UbaDetoursHandle = dlopen("libUbaDetours.dylib", RTLD_LAZY))
+					{
+						RequestNextProcess = (UbaRequestNextProcessFunc*)(void*)dlsym(UbaDetoursHandle, "UbaRequestNextProcess");
+					}
+#endif
+					if (!RequestNextProcess)
+					{
+						break;
+					}
+				}
 
 				// Request new process
-				TCHAR Arguments[1024];
-				if (!RequestNextProcess(0, Arguments, 1024))
+				ubachar Temp[1024];
+				if (!RequestNextProcess(0, Temp, 1024))
 				{
 					break; // No process available, exit loop
 				}
 
+				const TCHAR* Arguments;
+#if PLATFORM_WINDOWS
+				Arguments = Temp;
+#else
+				auto Temp2 = StringCast<TCHAR>(Temp);
+				Arguments = Temp2.Get();
+#endif
+
 				// We got a new process, change inputs and outputs and run again
 				
 				TArray<FString> Tokens;
 				TArray<FString> Switches;
 				FCommandLine::Parse(Arguments, Tokens, Switches);
+				if (Tokens.Num() < 5)
+				{
+					UE_LOG(LogShaders, Error, TEXT("Did not get enough arguments for reuse: %s"), Arguments);
+					break;
+				}
 
 				WorkingDirectory = Tokens[0];
 				InputFilename = Tokens[3];
@@ -347,12 +383,6 @@ public:
 				CrashOutputFile = OutputFilePath;
 				continue;
 			}
-#else
-			if (IsUsingUBA())
-			{
-				break;
-			}
-#endif
 
 			if (TimeToLive == 0)
 			{
