[Shaders] Allow dumping more ShaderCompileWorker diagnostics per batch to CSV file.

This replaces the single diagnostic output of an error code with FShaderCompileWorkerDiagnostics struct which contains the error code, batch index, and several time stamps.
Can be enabled via CVar "r.ShaderCompiler.DumpWorkerDiagnostics" to dump all shader job batch diagnostics to a CSV file in ShaderDebugInfo/ folder.

#rnx

diff --git a/Engine/Source/Programs/ShaderCompileWorker/Private/ShaderCompileWorker.cpp b/Engine/Source/Programs/ShaderCompileWorker/Private/ShaderCompileWorker.cpp
--- a/Engine/Source/Programs/ShaderCompileWorker/Private/ShaderCompileWorker.cpp
+++ b/Engine/Source/Programs/ShaderCompileWorker/Private/ShaderCompileWorker.cpp
@@ -28,7 +28,7 @@
 
 #define DEBUG_USING_CONSOLE	0
 
-static double LastCompileTime = 0.0;
+static double GLastCompileTime = 0.0;
 static int32 GNumProcessedJobs = 0;
 
 enum class EXGEMode
@@ -45,7 +45,7 @@ inline bool IsUsingXGE()
 	return GXGEMode != EXGEMode::None;
 }
 
-static FSCWErrorCode::ECode GFailedErrorCode = FSCWErrorCode::Success;
+static FShaderCompileWorkerDiagnostics GWorkerDiagnostics;
 
 static void OnXGEJobCompleted(const TCHAR* WorkingDirectory)
 {
@@ -89,7 +89,7 @@ inline bool IsUsingUBA()
 
 static inline void ExitWithoutCrash(FSCWErrorCode::ECode ErrorCode, const FString& Message)
 {
-	GFailedErrorCode = ErrorCode;
+	GWorkerDiagnostics.ErrorCode = ErrorCode;
 	FCString::Snprintf(GErrorExceptionDescription, sizeof(GErrorExceptionDescription), TEXT("%s"), *Message);
 	UE_LOG(LogShaders, Fatal, TEXT("%s"), *Message);
 }
@@ -147,7 +147,7 @@ static const TCHAR* GetLocalHostname(int32* OutHostnameLength = nullptr)
 	return *Hostname;
 }
 
-static int64 WriteOutputFileHeader(FArchive& OutputFile, int32 ErrorCode, int32 CallstackLength, const TCHAR* Callstack,
+static int64 WriteOutputFileHeader(FArchive& OutputFile, int32 CallstackLength, const TCHAR* Callstack,
 	int32 ExceptionInfoLength, const TCHAR* ExceptionInfo)
 {
 	TRACE_CPUPROFILER_EVENT_SCOPE(WriteOutputFileHeader);
@@ -161,7 +161,7 @@ static int64 WriteOutputFileHeader(FArchive& OutputFile, int32 ErrorCode, int32
 	FileSizePosition = OutputFile.Tell();
 	OutputFile << FileSize;
 
-	OutputFile << ErrorCode;
+	OutputFile << GWorkerDiagnostics;
 
 	OutputFile << GNumProcessedJobs;
 
@@ -174,7 +174,7 @@ static int64 WriteOutputFileHeader(FArchive& OutputFile, int32 ErrorCode, int32
 	const TCHAR* Hostname = GetLocalHostname(&HostnameLength);
 	OutputFile << HostnameLength;
 
-	if (ErrorCode != FSCWErrorCode::Success)
+	if (GWorkerDiagnostics.ErrorCode != FSCWErrorCode::Success)
 	{
 		if (CallstackLength > 0)
 		{
@@ -192,7 +192,7 @@ static int64 WriteOutputFileHeader(FArchive& OutputFile, int32 ErrorCode, int32
 		}
 
 		// Store available and used physical memory of host machine on OOM error
-		if (ErrorCode == FSCWErrorCode::OutOfMemory)
+		if (GWorkerDiagnostics.ErrorCode == FSCWErrorCode::OutOfMemory)
 		{
 			FPlatformMemoryStats MemoryStats = FPlatformMemory::GetStats();
 
@@ -279,11 +279,17 @@ public:
 					break;
 				}
 
+				// Record time since SCW entered the entry point
+				const double BatchProcessStartTime = FPlatformTime::Seconds();
+				GWorkerDiagnostics.BatchPreparationTime = BatchProcessStartTime - (NumberProcessed > 0 ? GLastCompileTime : GWorkerDiagnostics.EntryPointTimestamp);
+				GWorkerDiagnostics.BatchIndex = NumberProcessed;
+
 				UE_LOG(LogShaders, Log, TEXT("Processing shader"));
 
 				ProcessInputFromArchive(InputFilePtr, SingleJobs, PipelineJobs, PipelineJobNames);
 
-				LastCompileTime = FPlatformTime::Seconds();
+				GLastCompileTime = FPlatformTime::Seconds();
+				GWorkerDiagnostics.BatchProcessTime = GLastCompileTime - BatchProcessStartTime;
 
 				// Close the input file.
 				delete InputFilePtr;
@@ -800,9 +806,17 @@ private:
 		TRACE_CPUPROFILER_EVENT_SCOPE(WriteToOutputArchive);
 
 		FArchive& OutputFile = *OutputFilePtr;
-		const int64 FileSizePosition = FSCWErrorCode::IsSet()
-			? WriteOutputFileHeader(OutputFile, FSCWErrorCode::Get(), 0, nullptr, FSCWErrorCode::GetInfo().Len(), *FSCWErrorCode::GetInfo())
-			: WriteOutputFileHeader(OutputFile, FSCWErrorCode::Success, 0, nullptr, 0, nullptr);
+		int64 FileSizePosition = 0;
+		if (FSCWErrorCode::IsSet())
+		{
+			GWorkerDiagnostics.ErrorCode = FSCWErrorCode::Get();
+			FileSizePosition = WriteOutputFileHeader(OutputFile, 0, nullptr, FSCWErrorCode::GetInfo().Len(), *FSCWErrorCode::GetInfo());
+		}
+		else
+		{
+			GWorkerDiagnostics.ErrorCode = FSCWErrorCode::Success;
+			FileSizePosition = WriteOutputFileHeader(OutputFile, 0, nullptr, 0, nullptr);
+		}
 
 		{
 			int32 SingleJobHeader = ShaderCompileWorkerSingleJobHeader;
@@ -865,9 +879,9 @@ private:
 		}
 
 		const double CurrentTime = FPlatformTime::Seconds();
-		if (CurrentTime - LastCompileTime > TimeToLive)
+		if (CurrentTime - GLastCompileTime > TimeToLive)
 		{
-			UE_LOG(LogShaders, Log, TEXT("No jobs found for %f seconds, exiting"), (float)(CurrentTime - LastCompileTime));
+			UE_LOG(LogShaders, Log, TEXT("No jobs found for %f seconds, exiting"), (float)(CurrentTime - GLastCompileTime));
 			FPlatformMisc::RequestExit(false);
 		}
 #else
@@ -906,9 +920,9 @@ private:
 
 			const double CurrentTime = FPlatformTime::Seconds();
 			// If we have been idle for 20 seconds then exit
-			if (CurrentTime - LastCompileTime > TimeToLive)
+			if (CurrentTime - GLastCompileTime > TimeToLive)
 			{
-				UE_LOG(LogShaders, Log, TEXT("No jobs found for %f seconds, exiting"), (float)(CurrentTime - LastCompileTime));
+				UE_LOG(LogShaders, Log, TEXT("No jobs found for %f seconds, exiting"), (float)(CurrentTime - GLastCompileTime));
 				FPlatformMisc::RequestExit(false);
 			}
 		}
@@ -1191,7 +1205,7 @@ static int32 GuardedMain(int32 argc, TCHAR* argv[], FString& CrashOutputFile, bo
 		}
 	}
 
-	LastCompileTime = FPlatformTime::Seconds();
+	GLastCompileTime = FPlatformTime::Seconds();
 
 	if (bDirectMode)
 	{
@@ -1259,22 +1273,21 @@ static int32 GuardedMainWrapper(int32 ArgC, TCHAR* ArgV[], FString& CrashOutputF
 
 			FArchive& OutputFile = *IFileManager::Get().CreateFileWriter(*CrashOutputFile, FILEWRITE_EvenIfReadOnly);
 
-			if (GFailedErrorCode == FSCWErrorCode::Success)
+			if (GWorkerDiagnostics.ErrorCode == FSCWErrorCode::Success)
 			{
 				if (FSCWErrorCode::IsSet())
 				{
 					// Use the value set inside the shader format
-					GFailedErrorCode = FSCWErrorCode::Get();
+					GWorkerDiagnostics.ErrorCode = FSCWErrorCode::Get();
 				}
 				else
 				{
 					// Something else failed before we could set the error code, so mark it as a General Crash
-					GFailedErrorCode = FSCWErrorCode::GeneralCrash;
+					GWorkerDiagnostics.ErrorCode = FSCWErrorCode::GeneralCrash;
 				}
 			}
 
-			int64 FileSizePosition = WriteOutputFileHeader(OutputFile, GFailedErrorCode, ExceptionCallStack.Len(), *ExceptionCallStack,
-				ExceptionMsg.Len(), *ExceptionMsg);
+			int64 FileSizePosition = WriteOutputFileHeader(OutputFile, ExceptionCallStack.Len(), *ExceptionCallStack, ExceptionMsg.Len(), *ExceptionMsg);
 
 			int32 NumBatches = 0;
 			OutputFile << NumBatches;
@@ -1292,7 +1305,7 @@ static int32 GuardedMainWrapper(int32 ArgC, TCHAR* ArgV[], FString& CrashOutputF
 			}
 			else if (GetUbaModule())
 			{
-				ReturnCode = GFailedErrorCode;
+				ReturnCode = GWorkerDiagnostics.ErrorCode;
 			}
 		}
 	}
@@ -1317,6 +1330,9 @@ IMPLEMENT_APPLICATION(ShaderCompileWorker, "ShaderCompileWorker")
 
 INT32_MAIN_INT32_ARGC_TCHAR_ARGV()
 {
+	// Record timestamp when entry point is entered
+	GWorkerDiagnostics.EntryPointTimestamp = FPlatformTime::Seconds();
+
 	// Redirect for special XGE utilities...
 	extern bool XGEMain(int ArgC, TCHAR* ArgV[], int32& ReturnCode);
 	{
diff --git a/Engine/Source/Runtime/Engine/Private/ShaderCompiler/ShaderCompiler.cpp b/Engine/Source/Runtime/Engine/Private/ShaderCompiler/ShaderCompiler.cpp
--- a/Engine/Source/Runtime/Engine/Private/ShaderCompiler/ShaderCompiler.cpp
+++ b/Engine/Source/Runtime/Engine/Private/ShaderCompiler/ShaderCompiler.cpp
@@ -1502,7 +1502,7 @@ void FShaderCompilingManager::ProcessFinishedJob(FShaderCommonCompileJob* Finish
 }
 
 /** Launches the worker, returns the launched process handle. */
-FProcHandle FShaderCompilingManager::LaunchWorker(const FString& WorkingDirectory, uint32 InProcessId, uint32 ThreadId, const FString& WorkerInputFile, const FString& WorkerOutputFile)
+FProcHandle FShaderCompilingManager::LaunchWorker(const FString& WorkingDirectory, uint32 InParentProcessId, uint32 ThreadId, const FString& WorkerInputFile, const FString& WorkerOutputFile, uint32* OutWorkerProcessId)
 {
 	// don't allow any jobs if not allowed
 	if (!AllowShaderCompiling())
@@ -1517,7 +1517,7 @@ FProcHandle FShaderCompilingManager::LaunchWorker(const FString& WorkingDirector
 	// Use IFileManager to do path conversion to properly handle sandbox paths (outside of standard paths in particular).
 	FString WorkerAbsoluteDirectory = IFileManager::Get().ConvertToAbsolutePathForExternalAppForWrite(*WorkingDirectory);
 	FPaths::NormalizeDirectoryName(WorkerAbsoluteDirectory);
-	FString WorkerParameters = FString(TEXT("\"")) + WorkerAbsoluteDirectory + TEXT("/\" ") + FString::FromInt(InProcessId) + TEXT(" ") + FString::FromInt(ThreadId) + TEXT(" ") + WorkerInputFile + TEXT(" ") + WorkerOutputFile;
+	FString WorkerParameters = FString(TEXT("\"")) + WorkerAbsoluteDirectory + TEXT("/\" ") + FString::FromInt(InParentProcessId) + TEXT(" ") + FString::FromInt(ThreadId) + TEXT(" ") + WorkerInputFile + TEXT(" ") + WorkerOutputFile;
 	WorkerParameters += FString(TEXT(" -communicatethroughfile "));
 	if ( GIsBuildMachine )
 	{
@@ -1598,6 +1598,10 @@ FProcHandle FShaderCompilingManager::LaunchWorker(const FString& WorkingDirector
 		FProcHandle WorkerHandle = FPlatformProcess::CreateProc(*ShaderCompileWorkerName, *WorkerParameters, true, false, false, &WorkerId, PriorityModifier, nullptr, nullptr);
 		if (WorkerHandle.IsValid())
 		{
+			if (OutWorkerProcessId)
+			{
+				*OutWorkerProcessId = WorkerId;
+			}
 			// Process launched at least once successfully
 			bFirstLaunch = false;
 		}
diff --git a/Engine/Source/Runtime/Engine/Private/ShaderCompiler/ShaderCompilerDistributed.cpp b/Engine/Source/Runtime/Engine/Private/ShaderCompiler/ShaderCompilerDistributed.cpp
--- a/Engine/Source/Runtime/Engine/Private/ShaderCompiler/ShaderCompilerDistributed.cpp
+++ b/Engine/Source/Runtime/Engine/Private/ShaderCompiler/ShaderCompilerDistributed.cpp
@@ -43,6 +43,8 @@ namespace DistributedShaderCompilerVariables
 
 }
 
+extern bool GShaderCompilerDumpWorkerDiagnostics;
+
 bool FShaderCompileDistributedThreadRunnable_Interface::IsSupported()
 {
 	//TODO Handle Generic response
@@ -442,9 +444,17 @@ int32 FShaderCompileDistributedThreadRunnable_Interface::CompilingLoop()
 				if (TUniquePtr<FArchive> OutputFileAr = TUniquePtr<FArchive>(IFileManager::Get().CreateFileReader(*Task->OutputFilePath, FILEREAD_Silent)))
 				{
 					bOutputFileReadFailed = false;
-					if (FShaderCompileUtilities::DoReadTaskResults(Task->ShaderJobs, *OutputFileAr) == FSCWErrorCode::Success)
+					FShaderCompileWorkerDiagnostics WorkerDiagnostics;
+					if (FShaderCompileUtilities::DoReadTaskResults(Task->ShaderJobs, *OutputFileAr, GShaderCompilerDumpWorkerDiagnostics ? &WorkerDiagnostics : nullptr) == FSCWErrorCode::Success)
 					{
 						bCompileJobsSucceeded = true;
+
+						if (GShaderCompilerDumpWorkerDiagnostics)
+						{
+							FString BatchLabel = FPaths::GetCleanFilename(Task->InputFilePath);
+							constexpr uint32 UnavailableWorkerId = 0;
+							GShaderCompilerStats->RegisterWorkerDiagnostics(WorkerDiagnostics, MoveTemp(BatchLabel), Task->ShaderJobs.Num(), UnavailableWorkerId);
+						}
 					}
 				}
 			}
diff --git a/Engine/Source/Runtime/Engine/Private/ShaderCompiler/ShaderCompilerStats.cpp b/Engine/Source/Runtime/Engine/Private/ShaderCompiler/ShaderCompilerStats.cpp
--- a/Engine/Source/Runtime/Engine/Private/ShaderCompiler/ShaderCompilerStats.cpp
+++ b/Engine/Source/Runtime/Engine/Private/ShaderCompiler/ShaderCompilerStats.cpp
@@ -39,6 +39,13 @@ static FAutoConsoleVariableRef CVarDumpShaderTimeStats(
 	GDumpShaderTimeStats,
 	TEXT("When set to true, dump shader compiler timing statistics to a CSV file."));
 
+bool GShaderCompilerDumpWorkerDiagnostics = false;
+static FAutoConsoleVariableRef CVarDumpWorkerDiagnostics(
+	TEXT("r.ShaderCompiler.DumpWorkerDiagnostics"),
+	GShaderCompilerDumpWorkerDiagnostics,
+	TEXT("If enabled, the shader compiler will dump CSV files into the ShaderDebugInfo folder with diagnostics for each batch of shader compile jobs."),
+	ECVF_ReadOnly);
+
 
 FShaderCompilerStats* GShaderCompilerStats = nullptr;
 
@@ -212,6 +219,7 @@ static FString PrintJobsCompletedPercentageToString(int64 JobsAssigned, int64 Jo
 
 static void DumpShaderTimingsToCSVFile(const TCHAR* Filename, const TMap<FString, FShaderTimings>& ShaderTimings)
 {
+#if ALLOW_DEBUG_FILES
 	// Ensure output folder exists
 	const FString& OutputDirectory = GShaderCompilingManager->GetAbsoluteShaderDebugInfoDirectory();
 	if (!IFileManager::Get().DirectoryExists(*OutputDirectory))
@@ -219,23 +227,82 @@ static void DumpShaderTimingsToCSVFile(const TCHAR* Filename, const TMap<FString
 		IFileManager::Get().MakeDirectory(*OutputDirectory, true);
 	}
 
-	// Write CSV table
-	FString Table;
+	// Write CSV table to file
+	const FString OutputFilename = FPaths::Combine(OutputDirectory, Filename);
+	if (TUniquePtr<FArchive> OutputFile = TUniquePtr<FArchive>(IFileManager::Get().CreateFileWriter(*OutputFilename)))
+	{
+		FDiagnosticTableWriterCSV Table(OutputFile.Get());
 
-	Table += FString::Printf(TEXT("SHADER,NUMBER OF COMPILATIONS,TOTAL COMPILE TIME,TOTAL PREPROCESS TIME,AVERAGE TIME,MAX TIME,MIN TIME%s"), LINE_TERMINATOR);
+		Table.AddColumn(TEXT("SHADER"));
+		Table.AddColumn(TEXT("NUMBER OF COMPILATIONS"));
+		Table.AddColumn(TEXT("TOTAL COMPILE TIME (s)"));
+		Table.AddColumn(TEXT("TOTAL PREPROCESS TIME (s)"));
+		Table.AddColumn(TEXT("AVERAGE TIME (s)"));
+		Table.AddColumn(TEXT("MAX TIME (s)"));
+		Table.AddColumn(TEXT("MIN TIME (s)"));
+		Table.CycleRow();
 
-	for (TMap<FString, FShaderTimings>::TConstIterator Iter(ShaderTimings); Iter; ++Iter)
+		for (TMap<FString, FShaderTimings>::TConstIterator Iter(ShaderTimings); Iter; ++Iter)
+		{
+			const FShaderTimings& Timings = Iter.Value();
+			Table.AddColumn(TEXT("%s"), *Iter.Key());
+			Table.AddColumn(TEXT("%d"), Timings.NumCompiled);
+			Table.AddColumn(TEXT("%.2f"), Timings.TotalCompileTime);
+			Table.AddColumn(TEXT("%.2f"), Timings.TotalPreprocessTime);
+			Table.AddColumn(TEXT("%.2f"), Timings.AverageCompileTime);
+			Table.AddColumn(TEXT("%.2f"), Timings.MaxCompileTime);
+			Table.AddColumn(TEXT("%.2f"), Timings.MinCompileTime);
+			Table.CycleRow();
+		}
+	}
+#endif // ALLOW_DEBUG_FILES
+}
+
+static void DumpShaderWorkerDiagnosticsToCSVFile(const TCHAR* Filename, const TArray<FShaderCompilerStats::FWorkerDiagnosticsInfo>& InWorkerDiagnostics)
+{
+#if ALLOW_DEBUG_FILES
+	// Ensure output folder exists
+	const FString& OutputDirectory = GShaderCompilingManager->GetAbsoluteShaderDebugInfoDirectory();
+	if (!IFileManager::Get().DirectoryExists(*OutputDirectory))
 	{
-		const FShaderTimings& Timings = Iter.Value();
-		Table += FString::Printf(
-			TEXT("%s,%d,%4.2fs,%4.2fs,%4.2fs,%4.2fs,%4.2fs%s"),
-			*Iter.Key(), Timings.NumCompiled, Timings.TotalCompileTime, Timings.TotalPreprocessTime, Timings.AverageCompileTime, Timings.MaxCompileTime, Timings.MinCompileTime, LINE_TERMINATOR
-		);
+		IFileManager::Get().MakeDirectory(*OutputDirectory, true);
 	}
 
-	// Dump CSV table to file
+	// Write CSV table to file
 	const FString OutputFilename = FPaths::Combine(OutputDirectory, Filename);
-	FFileHelper::SaveStringToFile(Table, *OutputFilename);
+	if (TUniquePtr<FArchive> OutputFile = TUniquePtr<FArchive>(IFileManager::Get().CreateFileWriter(*OutputFilename)))
+	{
+		FDiagnosticTableWriterCSV Table(OutputFile.Get());
+
+		Table.AddColumn(TEXT("BATCH LABEL"));
+		Table.AddColumn(TEXT("BATCH INDEX"));
+		Table.AddColumn(TEXT("BATCH SIZE"));
+		Table.AddColumn(TEXT("WORKER ID"));
+		Table.AddColumn(TEXT("MAIN TIMESTAMP"));
+		Table.AddColumn(TEXT("PREPARATION (s)"));
+		Table.AddColumn(TEXT("DURATION (s)"));
+		Table.CycleRow();
+
+		for (const FShaderCompilerStats::FWorkerDiagnosticsInfo& Info : InWorkerDiagnostics)
+		{
+			Table.AddColumn(TEXT("%s"), *Info.BatchLabel);
+			Table.AddColumn(TEXT("%d"), Info.WorkerDiagnosticsOutput.BatchIndex);
+			Table.AddColumn(TEXT("%d"), Info.BatchSize);
+			if (Info.WorkerId == 0)
+			{
+				Table.AddColumn(TEXT("n/a"));
+			}
+			else
+			{
+				Table.AddColumn(TEXT("%u"), Info.WorkerId);
+			}
+			Table.AddColumn(TEXT("%.2f"), Info.WorkerDiagnosticsOutput.EntryPointTimestamp);
+			Table.AddColumn(TEXT("%.2f"), Info.WorkerDiagnosticsOutput.BatchPreparationTime);
+			Table.AddColumn(TEXT("%.2f"), Info.WorkerDiagnosticsOutput.BatchProcessTime);
+			Table.CycleRow();
+		}
+	}
+#endif // ALLOW_DEBUG_FILES
 }
 
 void FShaderCompilerStats::WriteStatSummary()
@@ -436,6 +503,14 @@ void FShaderCompilerStats::WriteStatSummary()
 		}
 	}
 
+	if (WorkerDiagnostics.Num())
+	{
+		if (GShaderCompilerDumpWorkerDiagnostics)
+		{
+			DumpShaderWorkerDiagnosticsToCSVFile(TEXT("ShaderCompileWorker.Diagnostics.csv"), WorkerDiagnostics);
+		}
+	}
+
 	MaterialCounters.WriteStatSummary(AggregatedSuffix);
 
 	UE_LOG(LogShaderCompilers, Display, TEXT("================================================"));
@@ -621,6 +696,11 @@ void FShaderCompilerStats::Aggregate(FShaderCompilerStats& Other)
 	}
 
 	MaterialCounters += Other.MaterialCounters;
+
+	if (GShaderCompilerDumpWorkerDiagnostics)
+	{
+		WorkerDiagnostics.Append(Other.WorkerDiagnostics);
+	}
 }
 
 
@@ -991,6 +1071,16 @@ void FShaderCompilerStats::RegisterDistributedBuildStats(const FDistributedBuild
 	Counters.MaxActiveAgentCores = FMath::Max(Counters.MaxActiveAgentCores, InStats.MaxActiveAgentCores);
 }
 
+void FShaderCompilerStats::RegisterWorkerDiagnostics(const FShaderCompileWorkerDiagnostics& InDiagnostics, FString&& InBatchLabel, int32 InBatchSize, uint32 InWorkerId)
+{
+	FScopeLock Lock(&CompileStatsLock);
+	FWorkerDiagnosticsInfo& NewDiagnosticsInfo = WorkerDiagnostics.AddDefaulted_GetRef();
+	NewDiagnosticsInfo.WorkerDiagnosticsOutput = InDiagnostics;
+	NewDiagnosticsInfo.BatchLabel = MoveTemp(InBatchLabel);
+	NewDiagnosticsInfo.BatchSize = InBatchSize;
+	NewDiagnosticsInfo.WorkerId = InWorkerId;
+}
+
 void FShaderCompilerMaterialCounters::WriteStatSummary(const TCHAR* AggregatedSuffix)
 {
 	auto CalcTimePercentage = [&](double Val) {
diff --git a/Engine/Source/Runtime/Engine/Private/ShaderCompiler/ShaderCompilerThreadRunnable.cpp b/Engine/Source/Runtime/Engine/Private/ShaderCompiler/ShaderCompilerThreadRunnable.cpp
--- a/Engine/Source/Runtime/Engine/Private/ShaderCompiler/ShaderCompilerThreadRunnable.cpp
+++ b/Engine/Source/Runtime/Engine/Private/ShaderCompiler/ShaderCompilerThreadRunnable.cpp
@@ -42,6 +42,8 @@ static FAutoConsoleVariableRef CVarShaderCompilerMemoryLimit(
 	TEXT("By default 0, effectively disabling the limitation. If this is non-zero, it must be greater than or equal to 1024 since shader compilation must be granted at least 1024 MiB of memory in total."),
 	ECVF_ReadOnly);
 
+extern bool GShaderCompilerDumpWorkerDiagnostics;
+
 // Configuration to retry shader compile through workers after a worker has been abandoned
 static constexpr int32 GSingleThreadedRunsDisabled = -2;
 static constexpr int32 GSingleThreadedRunsIncreaseFactor = 8;
@@ -73,9 +75,12 @@ static void ApplyWorkerProcessMemoryLimits(const FProcHandle& Process)
 /** Information tracked for each shader compile worker process instance. */
 struct FShaderCompileWorkerInfo
 {
-	/** Process handle of the worker app once launched.  Invalid handle means no process. */
+	/** Process handle of the worker app once launched. Invalid handle means no process. */
 	FProcHandle WorkerProcess;
 
+	/** Process ID of the worker app once launched. Zero means no process. */
+	uint32 WorkerProcessId;
+
 	/** Tracks whether tasks have been issued to the worker. */
 	bool bIssuedTasksToWorker;	
 
@@ -97,13 +102,18 @@ struct FShaderCompileWorkerInfo
 	/** Jobs that this worker is responsible for compiling. */
 	TArray<FShaderCommonCompileJobPtr> QueuedJobs;
 
+	/** Current batch ID to uniquely identify each batch for diagnostics output. */
+	uint64 QueuedBatchId;
+
 	FShaderCompileWorkerInfo() :
+		WorkerProcessId(0),
 		bIssuedTasksToWorker(false),
 		bLaunchedWorker(false),
 		bComplete(false),
 		bAvailable(true),
 		StartTime(0.0),
-		FinishTime(0.0)
+		FinishTime(0.0),
+		QueuedBatchId(0)
 	{
 	}
 
@@ -264,6 +274,8 @@ uint32 FShaderCompileThreadRunnableBase::Run()
 	return 0;
 }
 
+static std::atomic_uint64_t GLocalShaderCompileBatchCounter = 0;
+
 int32 FShaderCompileThreadRunnable::PullTasksFromQueue()
 {
 	TRACE_CPUPROFILER_EVENT_SCOPE(FShaderCompileThreadRunnable::PullTasksFromQueue);
@@ -275,6 +287,7 @@ int32 FShaderCompileThreadRunnable::PullTasksFromQueue()
 			WorkerInfo.bIssuedTasksToWorker = false;
 			WorkerInfo.bLaunchedWorker = false;
 			WorkerInfo.StartTime = FPlatformTime::Seconds();
+			WorkerInfo.QueuedBatchId = ++GLocalShaderCompileBatchCounter;
 
 			if (WorkerInfo.FinishTime > 0.0)
 			{
@@ -651,7 +664,7 @@ bool FShaderCompileThreadRunnable::LaunchWorkersIfNeeded()
 				const FString WorkingDirectory = GetWorkingDirectoryForWorker(WorkerIndex, bRelativePath);
 
 				// Store the handle with this thread so that we will know not to launch it again
-				CurrentWorkerInfo.WorkerProcess = Manager->LaunchWorker(WorkingDirectory, Manager->ProcessId, WorkerIndex, GWorkerInputFilename, GWorkerOutputFilename);
+				CurrentWorkerInfo.WorkerProcess = Manager->LaunchWorker(WorkingDirectory, Manager->ProcessId, WorkerIndex, GWorkerInputFilename, GWorkerOutputFilename, &CurrentWorkerInfo.WorkerProcessId);
 				CurrentWorkerInfo.bLaunchedWorker = true;
 
 				// Assign process to job object to monitor the total memory consumption of all SCW processes. Ignore if we only estimate the memory due to partial support in POSIX/Wine.
@@ -718,7 +731,15 @@ int32 FShaderCompileThreadRunnable::ReadAvailableResults()
 				if (TUniquePtr<FArchive> OutputFile = TUniquePtr<FArchive>(IFileManager::Get().CreateFileReader(*OutputFileNameAndPath, FILEREAD_Silent)))
 				{
 					check(!CurrentWorkerInfo.bComplete);
-					FShaderCompileUtilities::DoReadTaskResults(CurrentWorkerInfo.QueuedJobs, *OutputFile);
+					FShaderCompileWorkerDiagnostics WorkerDiagnostics;
+
+					FShaderCompileUtilities::DoReadTaskResults(CurrentWorkerInfo.QueuedJobs, *OutputFile, GShaderCompilerDumpWorkerDiagnostics ? &WorkerDiagnostics : nullptr);
+
+					if (GShaderCompilerDumpWorkerDiagnostics)
+					{
+						FString BatchLabel = FString::Printf(TEXT("Local-%d"), CurrentWorkerInfo.QueuedBatchId);
+						GShaderCompilerStats->RegisterWorkerDiagnostics(WorkerDiagnostics, MoveTemp(BatchLabel), CurrentWorkerInfo.QueuedJobs.Num(), CurrentWorkerInfo.WorkerProcessId);
+					}
 
 					// Close the output file.
 					OutputFile.Reset();
diff --git a/Engine/Source/Runtime/Engine/Private/ShaderCompiler/ShaderCompilerUtil.cpp b/Engine/Source/Runtime/Engine/Private/ShaderCompiler/ShaderCompilerUtil.cpp
--- a/Engine/Source/Runtime/Engine/Private/ShaderCompiler/ShaderCompilerUtil.cpp
+++ b/Engine/Source/Runtime/Engine/Private/ShaderCompiler/ShaderCompilerUtil.cpp
@@ -766,7 +766,7 @@ struct FSCWOutputFileContext
 
 // Process results from Worker Process.
 // Returns false if reading the tasks failed but we were able to recover from handing a crash report. In this case, all jobs must be submitted/processed again.
-FSCWErrorCode::ECode FShaderCompileUtilities::DoReadTaskResults(const TArray<FShaderCommonCompileJobPtr>& QueuedJobs, FArchive& OutputFile)
+FSCWErrorCode::ECode FShaderCompileUtilities::DoReadTaskResults(const TArray<FShaderCommonCompileJobPtr>& QueuedJobs, FArchive& OutputFile, FShaderCompileWorkerDiagnostics* OutWorkerDiagnostics)
 {
 	FSCWOutputFileContext OutputFileContext(OutputFile);
 
@@ -791,8 +791,13 @@ FSCWErrorCode::ECode FShaderCompileUtilities::DoReadTaskResults(const TArray<FSh
 		ShaderCompileWorkerError::HandleOutputFileCorrupted(*OutputFile.GetArchiveName(), OutputFileContext.FileSize, OutputFile.TotalSize());
 	}
 
-	int32 ErrorCode = 0;
-	OutputFile << ErrorCode;
+	FShaderCompileWorkerDiagnostics WorkerDiagnostics;
+	OutputFile << WorkerDiagnostics;
+
+	if (OutWorkerDiagnostics)
+	{
+		*OutWorkerDiagnostics = WorkerDiagnostics;
+	}
 
 	int32 NumProcessedJobs = 0;
 	OutputFile << NumProcessedJobs;
@@ -806,13 +811,13 @@ FSCWErrorCode::ECode FShaderCompileUtilities::DoReadTaskResults(const TArray<FSh
 	int32 HostnameLength = 0;
 	OutputFile << HostnameLength;
 
-	if (ErrorCode != FSCWErrorCode::Success)
+	if (WorkerDiagnostics.ErrorCode != FSCWErrorCode::Success)
 	{
 		// If worker crashed in a way we were able to recover from, return and expect the compile jobs to be reissued already
-		if (HandleWorkerCrash(QueuedJobs, OutputFile, OutputVersion, OutputFileContext.FileSize, (FSCWErrorCode::ECode)ErrorCode, NumProcessedJobs, CallstackLength, ExceptionInfoLength, HostnameLength))
+		if (HandleWorkerCrash(QueuedJobs, OutputFile, OutputVersion, OutputFileContext.FileSize, (FSCWErrorCode::ECode)WorkerDiagnostics.ErrorCode, NumProcessedJobs, CallstackLength, ExceptionInfoLength, HostnameLength))
 		{
 			FSCWErrorCode::Reset();
-			return (FSCWErrorCode::ECode)ErrorCode;
+			return (FSCWErrorCode::ECode)WorkerDiagnostics.ErrorCode;
 		}
 	}
 
diff --git a/Engine/Source/Runtime/Engine/Public/ShaderCompiler.h b/Engine/Source/Runtime/Engine/Public/ShaderCompiler.h
--- a/Engine/Source/Runtime/Engine/Public/ShaderCompiler.h
+++ b/Engine/Source/Runtime/Engine/Public/ShaderCompiler.h
@@ -310,7 +310,7 @@ class FShaderCompileUtilities
 {
 public:
 	static bool DoWriteTasks(const TArray<FShaderCommonCompileJobPtr>& QueuedJobs, FArchive& TransferFile, IDistributedBuildController* BuildDistributionController = nullptr, bool bUseRelativePaths = false, bool bCompressTaskFile = false);
-	static FSCWErrorCode::ECode DoReadTaskResults(const TArray<FShaderCommonCompileJobPtr>& QueuedJobs, FArchive& OutputFile);
+	static FSCWErrorCode::ECode DoReadTaskResults(const TArray<FShaderCommonCompileJobPtr>& QueuedJobs, FArchive& OutputFile, FShaderCompileWorkerDiagnostics* OutWorkerDiagnostics = nullptr);
 
 	/** Execute the specified (single or pipeline) shader compile job. */
 	static void ExecuteShaderCompileJob(FShaderCommonCompileJob& Job);
@@ -762,6 +762,14 @@ public:
 		Distributed
 	};
 
+	struct FWorkerDiagnosticsInfo
+	{
+		FShaderCompileWorkerDiagnostics WorkerDiagnosticsOutput;
+		FString BatchLabel;
+		int32 BatchSize = 0;
+		uint32 WorkerId = 0; // Worker process ID or 0 if it came from a distributed job
+	};
+
 	/** Informs statistics about a time a local ShaderCompileWorker spent idle. */
 	void RegisterLocalWorkerIdleTime(double IdleTime);
 
@@ -780,6 +788,9 @@ public:
 	/** Informs about current distributed build statistics. */
 	void RegisterDistributedBuildStats(const FDistributedBuildStats& InStats);
 
+	/** Informs statistics about a new worker diagnostics for a finished job batch. */
+	void RegisterWorkerDiagnostics(const FShaderCompileWorkerDiagnostics& InDiagnostics, FString&& InBatchLabel, int32 InBatchSize, uint32 InWorkerId);
+
 	ENGINE_API void GatherAnalytics(const FString& BaseName, TArray<FAnalyticsEventAttribute>& Attributes);
 
 private:
@@ -796,6 +807,9 @@ private:
 	/** Map of shader names to their compilation timings */
 	TMap<FString, FShaderTimings> ShaderTimings;
 
+	/** Array of diagnostics information per batch. */
+	TArray<FWorkerDiagnosticsInfo> WorkerDiagnostics;
+
 	bool bMultiProcessAggregated = false;
 };
 
@@ -935,7 +949,7 @@ private:
 	void CalculateNumberOfCompilingThreads(int32 NumberOfCores, int32 NumberOfCoresIncludingHyperthreads);
 
 	/** Launches the worker, returns the launched process handle. */
-	FProcHandle LaunchWorker(const FString& WorkingDirectory, uint32 ProcessId, uint32 ThreadId, const FString& WorkerInputFile, const FString& WorkerOutputFile);
+	FProcHandle LaunchWorker(const FString& WorkingDirectory, uint32 InParentProcessId, uint32 ThreadId, const FString& WorkerInputFile, const FString& WorkerOutputFile, uint32* OutWorkerProcessId = nullptr);
 
 	/** Blocks on completion of the given shader maps. */
 	void BlockOnShaderMapCompletion(const TArray<int32>& ShaderMapIdsToFinishCompiling, TMap<int32, FShaderMapFinalizeResults>& CompiledShaderMaps);
diff --git a/Engine/Source/Runtime/RenderCore/Public/ShaderCompilerCore.h b/Engine/Source/Runtime/RenderCore/Public/ShaderCompilerCore.h
--- a/Engine/Source/Runtime/RenderCore/Public/ShaderCompilerCore.h
+++ b/Engine/Source/Runtime/RenderCore/Public/ShaderCompilerCore.h
@@ -30,7 +30,7 @@ typedef TSharedPtr<TArray<ANSICHAR>, ESPMode::ThreadSafe> FShaderSharedAnsiStrin
 // this is for the protocol, not the data, bump if FShaderCompilerInput/FShaderPreprocessOutput serialization, SerializeWorkerInput or ProcessInputFromArchive changes.
 inline const int32 ShaderCompileWorkerInputVersion = 29;
 // this is for the protocol, not the data, bump if FShaderCompilerOutput or WriteToOutputArchive changes.
-inline const int32 ShaderCompileWorkerOutputVersion = 25;
+inline const int32 ShaderCompileWorkerOutputVersion = 26;
 // this is for the protocol, not the data.
 inline const int32 ShaderCompileWorkerSingleJobHeader = 'S';
 // this is for the protocol, not the data.
@@ -716,6 +716,31 @@ struct FSCWErrorCode
 	static RENDERCORE_API bool IsSet();
 };
 
+/** Serializable structure of diagnostic output from a SCW process. Include error code and timing statistics for the duration of a job batch. */
+struct FShaderCompileWorkerDiagnostics
+{
+	/** Error code return from a ShaderCompileWorker process that terminated abnormally (not for shader syntax/semantic errors). See FSCWErrorCode::ECode for valid values. */
+	int32 ErrorCode = 0;
+
+	/** Timestamp when the ShaderCompileWorker entered the main entry point. */
+	double EntryPointTimestamp = 0.0;
+
+	/** Time in seconds before this batch was being processed. This either starts from when the process launched or since the last batch was finished. */
+	double BatchPreparationTime = 0.0;
+
+	/** Time in seconds it took the ShaderCompileWorker to process the entire job batch. */
+	double BatchProcessTime = 0.0;
+
+	/** Index of the batch that was processed by the same worker process. Helps diagnose how many batches a worker has processed. */
+	int32 BatchIndex = 0;
+
+	friend FArchive& operator<<(FArchive& Ar, FShaderCompileWorkerDiagnostics& Output)
+	{
+		Ar << Output.ErrorCode << Output.EntryPointTimestamp << Output.BatchPreparationTime << Output.BatchProcessTime << Output.BatchIndex;
+		return Ar;
+	}
+};
+
 #if PLATFORM_WINDOWS
 extern RENDERCORE_API int HandleShaderCompileException(Windows::LPEXCEPTION_POINTERS Info, FString& OutExMsg, FString& OutCallStack);
 #endif
