Treat UnrealPak & ShaderCompileWorker as engine modules

#rnx

diff --git a/Engine/Source/Programs/ShaderCompileWorker/ShaderCompileWorker.Build.cs b/Engine/Source/Programs/ShaderCompileWorker/ShaderCompileWorker.Build.cs
--- a/Engine/Source/Programs/ShaderCompileWorker/ShaderCompileWorker.Build.cs
+++ b/Engine/Source/Programs/ShaderCompileWorker/ShaderCompileWorker.Build.cs
@@ -7,6 +7,8 @@ public class ShaderCompileWorker : ModuleRules
 {
 	public ShaderCompileWorker(ReadOnlyTargetRules Target) : base(Target)
 	{
+		bTreatAsEngineModule = true;
+
 		PrivateDependencyModuleNames.AddRange(
 			new string[] {
 				"Core",
diff --git a/Engine/Source/Programs/UnrealPak/UnrealPak.Build.cs b/Engine/Source/Programs/UnrealPak/UnrealPak.Build.cs
--- a/Engine/Source/Programs/UnrealPak/UnrealPak.Build.cs
+++ b/Engine/Source/Programs/UnrealPak/UnrealPak.Build.cs
@@ -6,6 +6,8 @@ public class UnrealPak : ModuleRules
 {
 	public UnrealPak(ReadOnlyTargetRules Target) : base(Target)
 	{
+		bTreatAsEngineModule = true;
+
 		PublicIncludePathModuleNames.Add("Launch");
 
 		PrivateDependencyModuleNames.AddRange(new string[] {
