[UBA] More diagnostics when UBA output file is invalid.

- Add more diagnostics about the directory and disk space when the UBA output is unexpectedly missing or invalid.
- Instead of emitting the entire generated VectorVM code to the log, dump it to the ShaderDebugInfo folder (or Horde artifact).
- Also apply some code cleaup and UE coding guide lines.

#rnx

diff --git a/Engine/Plugins/UbaController/Source/Private/UbaControllerModule.cpp b/Engine/Plugins/UbaController/Source/Private/UbaControllerModule.cpp
--- a/Engine/Plugins/UbaController/Source/Private/UbaControllerModule.cpp
+++ b/Engine/Plugins/UbaController/Source/Private/UbaControllerModule.cpp
@@ -245,9 +245,9 @@ void FUbaControllerModule::InitializeController()
 			IFileManager::Get().MakeDirectory(*WorkingDirectory, true);
 
 			// Pre-create the directories so we don't have to explicitly register them to uba later
-			for (int32 It = 0; It != UbaControllerModule::SubFolderCount; ++It)
+			for (int32 FolderIndex = 0; FolderIndex < UbaControllerModule::SubFolderCount; ++FolderIndex)
 			{
-				IFileManager::Get().MakeDirectory(*FString::Printf(TEXT("%s/%d"), *WorkingDirectory, It));
+				IFileManager::Get().MakeDirectory(*FPaths::Combine(WorkingDirectory, FString::FromInt(FolderIndex)));
 			}
 
 			if (UbaControllerModule::bDumpTraceFiles)
@@ -274,9 +274,9 @@ bool FUbaControllerModule::SupportsLocalWorkers()
 FString FUbaControllerModule::CreateUniqueFilePath()
 {
 	check(bSupported);
-	int32 ID = NextFileID++;
-	int32 FolderID = ID % UbaControllerModule::SubFolderCount; // We use sub folders to be nicer to file system (we can end up with 20000 files in one folder otherwise)
-	return FString::Printf(TEXT("%s/%d/%d.uba"), *WorkingDirectory, FolderID, ID);
+	const int32 FileID = NextFileID++;
+	const int32 FolderID = FileID % UbaControllerModule::SubFolderCount; // We use sub folders to be nicer to file system (we can end up with 20000 files in one folder otherwise)
+	return FPaths::Combine(WorkingDirectory, FString::FromInt(FolderID), FString::Printf(TEXT("%d.uba"), FileID));
 }
 
 TFuture<FDistributedBuildTaskResult> FUbaControllerModule::EnqueueTask(const FTaskCommandData& CommandData)
diff --git a/Engine/Plugins/UbaController/Source/Private/UbaJobProcessor.cpp b/Engine/Plugins/UbaController/Source/Private/UbaJobProcessor.cpp
--- a/Engine/Plugins/UbaController/Source/Private/UbaJobProcessor.cpp
+++ b/Engine/Plugins/UbaController/Source/Private/UbaJobProcessor.cpp
@@ -290,15 +290,14 @@ void FUbaJobProcessor::RunTaskWithUba(FDistributedBuildTask* Task)
 	Info->OutputFile = Data.OutputFileName;
 	Info->Task = Task;
 
-	static auto ExitedFunc = [](void* userData, const uba::ProcessHandle& ph)
+	static auto ExitedFunc = [](void* UserData, const uba::ProcessHandle& Process)
 		{
-			uint32 logLineIndex = 0;
-			while (const uba::tchar* line = ProcessHandle_GetLogLine(&ph, logLineIndex++))
+			for (uint32 LogLineIndex = 0; const uba::tchar* LogLine = ProcessHandle_GetLogLine(&Process, LogLineIndex); ++LogLineIndex)
 			{
-				UE_LOG(LogUbaController, Display, TEXT("%s"), UBASTRING_TO_TCHAR(line));
+				UE_LOG(LogUbaController, Display, TEXT("%s"), UBASTRING_TO_TCHAR(LogLine));
 			}
 
-			if (auto Info = (ExitedInfo*)userData) // It can be null if custom message has already handled all of them
+			if (ExitedInfo* Info = static_cast<ExitedInfo*>(UserData)) // It can be null if custom message has already handled all of them
 			{
 				IFileManager::Get().Delete(*Info->InputFile);
 				SessionServer_RegisterDeleteFile(Info->Processor->UbaSessionServer, TCHAR_TO_UBASTRING(*Info->InputFile));
@@ -587,8 +586,48 @@ bool FUbaJobProcessor::ProcessOutputFile(FDistributedBuildTask* CompileTask)
 	}
 	else
 	{
-		const FString OutputFileName = CompileTask != nullptr ? *CompileTask->CommandData.OutputFileName : TEXT("Invalid CompileTask, cannot retrieve name");
-		UE_LOG(LogUbaController, Error, TEXT("Output File [%s] is invalid or does not exist"), *OutputFileName);
+		// Error: Output file was expected but is missing or invalid. Print some diagnostics about disk space and the folder content.
+		if (CompileTask != nullptr)
+		{
+			TStringBuilder<1024> DiagnosticsInfo;
+
+			// Append diagnostics about output directory
+			const FString OutputFileDirectory = FPaths::GetPath(CompileTask->CommandData.OutputFileName);
+			if (FPaths::DirectoryExists(OutputFileDirectory))
+			{
+				int32 NumFilesInDirectory = 0;
+				IFileManager::Get().IterateDirectory(
+					*OutputFileDirectory,
+					[&NumFilesInDirectory](const TCHAR* InName, bool bIsDirectory) -> bool
+					{
+						if (!bIsDirectory)
+						{
+							++NumFilesInDirectory;
+						}
+						return true;
+					}
+				);
+				DiagnosticsInfo.Appendf(TEXT("\n - Directory \"%s\" exists and contains %d file(s)"), *OutputFileDirectory, NumFilesInDirectory);
+			}
+			else
+			{
+				DiagnosticsInfo.Appendf(TEXT("\n - Directory \"%s\" does not exist"), *OutputFileDirectory);
+			}
+
+			// Append diagnostics about disk space
+			uint64 TotalNumberOfBytes = 0;
+			uint64 NumberOfFreeBytes = 0;
+			if (FPlatformMisc::GetDiskTotalAndFreeSpace(OutputFileDirectory, TotalNumberOfBytes, NumberOfFreeBytes))
+			{
+				DiagnosticsInfo.Appendf(TEXT("\n - Disk space: %lld MiB, free %lld MiB"), (long long int)TotalNumberOfBytes >> 20, (long long int)NumberOfFreeBytes >> 20);
+			}
+
+			UE_LOG(LogUbaController, Error, TEXT("Output File [%s] is invalid or does not exist%s"), *CompileTask->CommandData.OutputFileName, *DiagnosticsInfo);
+		}
+		else
+		{
+			UE_LOG(LogUbaController, Error, TEXT("Invalid CompileTask, cannot retrieve name"));
+		}
 		return false;
 	}
 
diff --git a/Engine/Source/Developer/ShaderFormatVectorVM/Private/VectorVMShaderCompiler.cpp b/Engine/Source/Developer/ShaderFormatVectorVM/Private/VectorVMShaderCompiler.cpp
--- a/Engine/Source/Developer/ShaderFormatVectorVM/Private/VectorVMShaderCompiler.cpp
+++ b/Engine/Source/Developer/ShaderFormatVectorVM/Private/VectorVMShaderCompiler.cpp
@@ -83,15 +83,26 @@ bool CompileVectorVMShader(
 
 		VMCompilationOutput.Errors = Converted.GetData();
 
-		UE_LOG(LogVectorVMShaderCompiler, Warning, TEXT("Warnings while processing %s"), *Input.DebugGroupName);
-
-		PreprocessOutput.ForEachLine
-		(
-			[](FAnsiStringView Line, int32 LineIndex)
+		// Dump preprocessed and generated source to debug info folder (or artifacts)
+		FString VectorVMPreprocessedLines;
+		PreprocessOutput.ForEachLine(
+			[&VectorVMPreprocessedLines](FAnsiStringView Line, int32 LineIndex)
 			{
-				UE_LOG(LogVectorVMShaderCompiler, Display, TEXT("/*%d*/%.*hs"), LineIndex, Line.Len(), Line.GetData());
+				VectorVMPreprocessedLines.Appendf(TEXT("/*%d*/%.*hs"), LineIndex, Line.Len(), Line.GetData());
 			}
 		);
+		const FString DebugInfoPath = Input.GetOrCreateShaderDebugInfoPath();
+		FFileHelper::SaveStringToFile(VectorVMPreprocessedLines, *FPaths::Combine(DebugInfoPath, TEXT("VectorVMPreprocessOutput.txt")));
+
+		if (ShaderSource)
+		{
+			FFileHelper::SaveStringToFile(ANSI_TO_TCHAR(ShaderSource), *FPaths::Combine(DebugInfoPath, TEXT("VectorVMShaderOutput.usf")));
+		}
+
+		UE_LOG(LogVectorVMShaderCompiler, Warning, TEXT("Warnings while processing %s"), *Input.DebugGroupName);
+		UE_LOG(LogVectorVMShaderCompiler, Warning, TEXT("%s"), *VMCompilationOutput.Errors);
+
+		free(ErrorLog);
 	}
 	else
 	{
@@ -100,14 +111,8 @@ bool CompileVectorVMShader(
 
 	if (ShaderSource)
 	{
-		UE_LOG(LogVectorVMShaderCompiler, Warning, TEXT("%hs"), (const char*)ShaderSource);
 		free(ShaderSource);
 	}
-	if (VMCompilationOutput.Errors.Len() != 0)
-	{
-		UE_LOG(LogVectorVMShaderCompiler, Warning, TEXT("%s"), *VMCompilationOutput.Errors);
-		free(ErrorLog);
-	}
 
 	return bResult;
 }
diff --git a/Engine/Source/Programs/ShaderCompileWorker/Private/ShaderCompileWorker.cpp b/Engine/Source/Programs/ShaderCompileWorker/Private/ShaderCompileWorker.cpp
--- a/Engine/Source/Programs/ShaderCompileWorker/Private/ShaderCompileWorker.cpp
+++ b/Engine/Source/Programs/ShaderCompileWorker/Private/ShaderCompileWorker.cpp
@@ -263,7 +263,7 @@ public:
 
 		int32 NumberProcessed = 0;
 
-		while(true)
+		while (true)
 		{
 			TArray<FShaderCompileJob> SingleJobs;
 			TArray<FShaderPipelineCompileJob> PipelineJobs;
@@ -273,8 +273,8 @@ public:
 			{
 				TRACE_CPUPROFILER_EVENT_SCOPE(ReadInput);
 
-				FArchive* InputFilePtr = OpenInputFile();
-				if(!InputFilePtr)
+				TUniquePtr<FArchive> InputFilePtr = TUniquePtr<FArchive>(OpenInputFile());
+				if (!InputFilePtr)
 				{
 					break;
 				}
@@ -286,13 +286,10 @@ public:
 
 				UE_LOG(LogShaders, Log, TEXT("Processing shader"));
 
-				ProcessInputFromArchive(InputFilePtr, SingleJobs, PipelineJobs, PipelineJobNames);
+				ProcessInputFromArchive(InputFilePtr.Get(), SingleJobs, PipelineJobs, PipelineJobNames);
 
 				GLastCompileTime = FPlatformTime::Seconds();
 				GWorkerDiagnostics.BatchProcessTime = GLastCompileTime - BatchProcessStartTime;
-
-				// Close the input file.
-				delete InputFilePtr;
 			}
 
 			// Prepare for output
@@ -301,15 +298,16 @@ public:
 				// write to in-memory bytestream instead for debugging purposes
 				TArray<uint8> MemBlock;
 				FMemoryWriter MemWriter(MemBlock);
-				WriteToOutputArchive(&MemWriter, SingleJobs, PipelineJobs, PipelineJobNames);
+				WriteToOutputArchive(MemWriter, SingleJobs, PipelineJobs, PipelineJobNames);
 			}
 			else
 			{
-				FArchive* OutputFilePtr = CreateOutputArchive();
-				check(OutputFilePtr);
-				WriteToOutputArchive(OutputFilePtr, SingleJobs, PipelineJobs, PipelineJobNames);
-				// Close the output file.
-				delete OutputFilePtr;
+				// Write worker output file
+				{
+					TUniquePtr<FArchive> OutputFilePtr = TUniquePtr<FArchive>(CreateOutputArchive());
+					check(OutputFilePtr);
+					WriteToOutputArchive(*OutputFilePtr, SingleJobs, PipelineJobs, PipelineJobNames);
+				}
 
 				// Change the output file name to requested one
 				IFileManager::Get().Move(*OutputFilePath, *TempFilePath);
@@ -801,11 +799,10 @@ private:
 		return OutputFilePtr;
 	}
 
-	void WriteToOutputArchive(FArchive* OutputFilePtr, TArray<FShaderCompileJob>& SingleJobs, TArray<FShaderPipelineCompileJob>& PipelineJobs, TArray<FString>& PipelineNames)
+	void WriteToOutputArchive(FArchive& OutputFile, TArray<FShaderCompileJob>& SingleJobs, TArray<FShaderPipelineCompileJob>& PipelineJobs, TArray<FString>& PipelineNames)
 	{
 		TRACE_CPUPROFILER_EVENT_SCOPE(WriteToOutputArchive);
 
-		FArchive& OutputFile = *OutputFilePtr;
 		int64 FileSizePosition = 0;
 		if (FSCWErrorCode::IsSet())
 		{
@@ -1271,32 +1268,30 @@ static int32 GuardedMainWrapper(int32 ArgC, TCHAR* ArgV[], FString& CrashOutputF
 			// Put app into critical error mode to allow dumping logs from memory to disk
 			GIsCriticalError = true;
 
-			FArchive& OutputFile = *IFileManager::Get().CreateFileWriter(*CrashOutputFile, FILEWRITE_EvenIfReadOnly);
-
-			if (GWorkerDiagnostics.ErrorCode == FSCWErrorCode::Success)
+			if (TUniquePtr<FArchive> OutputFile = TUniquePtr<FArchive>(IFileManager::Get().CreateFileWriter(*CrashOutputFile, FILEWRITE_EvenIfReadOnly)))
 			{
-				if (FSCWErrorCode::IsSet())
+				if (GWorkerDiagnostics.ErrorCode == FSCWErrorCode::Success)
 				{
-					// Use the value set inside the shader format
-					GWorkerDiagnostics.ErrorCode = FSCWErrorCode::Get();
-				}
-				else
-				{
-					// Something else failed before we could set the error code, so mark it as a General Crash
-					GWorkerDiagnostics.ErrorCode = FSCWErrorCode::GeneralCrash;
+					if (FSCWErrorCode::IsSet())
+					{
+						// Use the value set inside the shader format
+						GWorkerDiagnostics.ErrorCode = FSCWErrorCode::Get();
+					}
+					else
+					{
+						// Something else failed before we could set the error code, so mark it as a General Crash
+						GWorkerDiagnostics.ErrorCode = FSCWErrorCode::GeneralCrash;
+					}
 				}
-			}
-
-			int64 FileSizePosition = WriteOutputFileHeader(OutputFile, ExceptionCallStack.Len(), *ExceptionCallStack, ExceptionMsg.Len(), *ExceptionMsg);
 
-			int32 NumBatches = 0;
-			OutputFile << NumBatches;
-			OutputFile << NumBatches;
+				int64 FileSizePosition = WriteOutputFileHeader(*OutputFile, ExceptionCallStack.Len(), *ExceptionCallStack, ExceptionMsg.Len(), *ExceptionMsg);
 
-			UpdateFileSize(OutputFile, FileSizePosition);
+				int32 NumBatches = 0;
+				*OutputFile << NumBatches;
+				*OutputFile << NumBatches;
 
-			// Close the output file.
-			delete &OutputFile;
+				UpdateFileSize(*OutputFile, FileSizePosition);
+			}
 
 			if (IsUsingXGE())
 			{
