[ShaderCompileWorker]
* Added ".tmp" to temp file which will eventually be moved into .out

diff --git a/Engine/Source/Programs/ShaderCompileWorker/Private/ShaderCompileWorker.cpp b/Engine/Source/Programs/ShaderCompileWorker/Private/ShaderCompileWorker.cpp
--- a/Engine/Source/Programs/ShaderCompileWorker/Private/ShaderCompileWorker.cpp
+++ b/Engine/Source/Programs/ShaderCompileWorker/Private/ShaderCompileWorker.cpp
@@ -779,7 +779,7 @@ private:
 		{
 			FGuid Guid;
 			FPlatformMisc::CreateGuid(Guid);
-			TempFilePath = WorkingDirectory + Guid.ToString();
+			TempFilePath = WorkingDirectory + Guid.ToString() + TEXT(".tmp");
 		} while (IFileManager::Get().FileSize(*TempFilePath) != INDEX_NONE);
 
 		const double StartTime2 = FPlatformTime::Seconds();
