Update shader compiler worker references for arm64 and arm64ec.

#rnx

diff --git a/Engine/Plugins/XGEController/Source/Private/XGEControllerModule.cpp b/Engine/Plugins/XGEController/Source/Private/XGEControllerModule.cpp
--- a/Engine/Plugins/XGEController/Source/Private/XGEControllerModule.cpp
+++ b/Engine/Plugins/XGEController/Source/Private/XGEControllerModule.cpp
@@ -27,7 +27,13 @@
 	#include "Windows/HideWindowsPlatformTypes.h"
 
 	#define XGE_CONTROL_WORKER_NAME		TEXT("XGEControlWorker")
+#ifdef _M_ARM64EC
+	#define XGE_CONTROL_WORKER_FILENAME	TEXT("XGEControlWorkerarm64ec.exe")
+#elif PLATFORM_CPU_ARM_FAMILY
+	#define XGE_CONTROL_WORKER_FILENAME	TEXT("XGEControlWorkerarm64.exe")
+#else
 	#define XGE_CONTROL_WORKER_FILENAME	TEXT("XGEControlWorker.exe")
+#endif
 
 FString GetControlWorkerBinariesPath()
 {
@@ -36,7 +42,7 @@ FString GetControlWorkerBinariesPath()
 
 FString GetControlWorkerExePath()
 {
-	return FPaths::EngineDir() / TEXT("Binaries/Win64/XGEControlWorker.exe");
+	return FPaths::EngineDir() / TEXT("Binaries/Win64/") / XGE_CONTROL_WORKER_FILENAME;
 }
 
 #else
diff --git a/Engine/Source/Programs/ShaderCompileWorker/ShaderCompileWorker.Target.cs b/Engine/Source/Programs/ShaderCompileWorker/ShaderCompileWorker.Target.cs
--- a/Engine/Source/Programs/ShaderCompileWorker/ShaderCompileWorker.Target.cs
+++ b/Engine/Source/Programs/ShaderCompileWorker/ShaderCompileWorker.Target.cs
@@ -18,15 +18,20 @@ public class ShaderCompileWorkerTarget : TargetRules
 
         if (bUseXGEController && (Target.Platform == UnrealTargetPlatform.Win64) && Configuration == UnrealTargetConfiguration.Development)
         {
-            // The interception interface in XGE requires that the parent and child processes have different filenames on disk.
-            // To avoid building an entire separate worker just for this, we duplicate the ShaderCompileWorker in a post build step.
-            const string SrcPath  = "$(BinaryDir)\\$(TargetName).exe";
-            const string DestPath = "$(BinaryDir)\\XGEControlWorker.exe";
+			foreach (UnrealArch Arch in Target.Architectures.Architectures)
+			{
+				string ArchName = Arch.bIsX64 ? "" : Arch.ToString();
 
-            PostBuildSteps.Add(string.Format("echo Copying {0} to {1}", SrcPath, DestPath));
-            PostBuildSteps.Add(string.Format("copy /Y /B \"{0}\" /B \"{1}\" >nul:", SrcPath, DestPath));
+				// The interception interface in XGE requires that the parent and child processes have different filenames on disk.
+				// To avoid building an entire separate worker just for this, we duplicate the ShaderCompileWorker in a post build step.
+				string SrcPath  = "$(BinaryDir)\\$(TargetName)" + ArchName + ".exe";
+				string DestPath = "$(BinaryDir)\\XGEControlWorker" + ArchName + ".exe";
 
-			AdditionalBuildProducts.Add(DestPath);
+				PostBuildSteps.Add(string.Format("echo Copying {0} to {1}", SrcPath, DestPath));
+				PostBuildSteps.Add(string.Format("copy /Y /B \"{0}\" /B \"{1}\" >nul:", SrcPath, DestPath));
+
+				AdditionalBuildProducts.Add(DestPath);
+			}
         }
 
 		// Turn off various third party features we don't need
diff --git a/Engine/Source/Runtime/Engine/Private/ShaderCompiler/ShaderCompiler.cpp b/Engine/Source/Runtime/Engine/Private/ShaderCompiler/ShaderCompiler.cpp
--- a/Engine/Source/Runtime/Engine/Private/ShaderCompiler/ShaderCompiler.cpp
+++ b/Engine/Source/Runtime/Engine/Private/ShaderCompiler/ShaderCompiler.cpp
@@ -739,7 +739,13 @@ FShaderCompilingManager::FShaderCompilingManager() :
 	WorkersBusyTime = 0;
 
 #if PLATFORM_WINDOWS
+#ifdef _M_ARM64EC
+	FString ExecutableName("ShaderCompileWorkerarm64ec.exe");
+#elif PLATFORM_CPU_ARM_FAMILY
+	FString ExecutableName("ShaderCompileWorkerarm64.exe");
+#else
 	FString ExecutableName("ShaderCompileWorker.exe");
+#endif
 #else
 	FString ExecutableName("ShaderCompileWorker");
 #endif
