[ShaderCompileWorker]
* Added support for uba and linux
* Added description to shutdown requests

diff --git a/Engine/Source/Programs/ShaderCompileWorker/Private/ShaderCompileWorker.cpp b/Engine/Source/Programs/ShaderCompileWorker/Private/ShaderCompileWorker.cpp
--- a/Engine/Source/Programs/ShaderCompileWorker/Private/ShaderCompileWorker.cpp
+++ b/Engine/Source/Programs/ShaderCompileWorker/Private/ShaderCompileWorker.cpp
@@ -24,6 +24,8 @@
 
 #if PLATFORM_MAC
 #include <mach-o/dyld.h>
+#elif PLATFORM_LINUX
+#include <link.h>
 #endif
 
 #define DEBUG_USING_CONSOLE	0
@@ -67,20 +69,24 @@ HMODULE GetUbaModule()
 
 inline bool IsUsingUBA()
 {
+	static bool UbaLoaded;
+	if (UbaLoaded)
+	{
+		return true;
+	}
+
 #if PLATFORM_WINDOWS // Currently only implemented for windows
-	return GetUbaModule() != nullptr;
+	UbaLoaded = GetUbaModule() != nullptr;
 #elif PLATFORM_MAC
-	for (int i = 0, e = _dyld_image_count(); i != e; i++)
+	for (int i = 0, e = _dyld_image_count() && !UbaLoaded; i != e; i++)
 	{
-		if (strstr(_dyld_get_image_name(i), "UbaDetours.dylib"))
-		{
-			return true;
-		}
+		UbaLoaded = strstr(_dyld_get_image_name(i), "UbaDetours.dylib") != nullptr;
 	}
-	return false;
-#else
-	return false;
-#endif	
+#elif PLATFORM_LINUX
+	dl_iterate_phdr([](struct dl_phdr_info *info, size_t size, void *data) { UbaLoaded = UbaLoaded || (info->dlpi_name && strstr(info->dlpi_name, "UbaDetours.so") != nullptr); return 0; }, nullptr);
+#endif
+
+return UbaLoaded;
 }
 
 #if USING_CODE_ANALYSIS
@@ -344,6 +350,11 @@ public:
 					{
 						RequestNextProcess = (UbaRequestNextProcessFunc*)(void*)dlsym(UbaDetoursHandle, "UbaRequestNextProcess");
 					}
+#elif PLATFORM_LINUX
+					if (void* UbaDetoursHandle = dlopen("libUbaDetours.so", RTLD_LAZY))
+					{
+						RequestNextProcess = (UbaRequestNextProcessFunc*)(void*)dlsym(UbaDetoursHandle, "UbaRequestNextProcess");
+					}
 #endif
 					if (!RequestNextProcess)
 					{
@@ -858,8 +869,7 @@ private:
 	{
 		if (!InputFilename.Contains(TEXT("Only")))
 		{
-			UE_LOG(LogShaders, Log, TEXT("InputFilename did not contain 'Only', exiting after one job."));
-			FPlatformMisc::RequestExit(false);
+			FPlatformMisc::RequestExit(false, TEXT("ShaderCompileWorker - InputFilename did not contain 'Only', exiting after one job."));
 		}
 
 #if PLATFORM_MAC || PLATFORM_LINUX
@@ -870,8 +880,7 @@ private:
 			{
 				FString FilePath = FString(WorkingDirectory) + InputFilename;
 				checkf(IFileManager::Get().FileSize(*FilePath) == INDEX_NONE, TEXT("Exiting due to the parent process no longer running! FilePath=%s"), *FilePath);
-				UE_LOG(LogShaders, Log, TEXT("Parent process no longer running, exiting"));
-				FPlatformMisc::RequestExit(false);
+				FPlatformMisc::RequestExit(false, TEXT("ShaderCompileWorker - Parent process no longer running, exiting"));
 			}
 		}
 
@@ -879,7 +888,7 @@ private:
 		if (CurrentTime - GLastCompileTime > TimeToLive)
 		{
 			UE_LOG(LogShaders, Log, TEXT("No jobs found for %f seconds, exiting"), (float)(CurrentTime - GLastCompileTime));
-			FPlatformMisc::RequestExit(false);
+			FPlatformMisc::RequestExit(false, TEXT("ShaderCompileWorker - No jobs found for seconds, exiting"));
 		}
 #else
 		// Don't do these if the debugger is present
@@ -896,8 +905,7 @@ private:
 				if (ParentProcessHandle == nullptr)
 				{
 					checkf(IFileManager::Get().FileSize(*FilePath) == INDEX_NONE, TEXT("Exiting due to OpenProcess(ParentProcessId) failing! FilePath=%s"), *FilePath);
-					UE_LOG(LogShaders, Log, TEXT("Couldn't OpenProcess, Parent process no longer running, exiting"));
-					FPlatformMisc::RequestExit(false);
+					FPlatformMisc::RequestExit(false, TEXT("ShaderCompileWorker - Couldn't OpenProcess, Parent process no longer running, exiting"));
 				}
 				else
 				{
@@ -908,8 +916,7 @@ private:
 					if (WaitResult != WAIT_TIMEOUT)
 					{
 						checkf(IFileManager::Get().FileSize(*FilePath) == INDEX_NONE, TEXT("Exiting due to WaitForSingleObject(ParentProcessHandle) signaling! FilePath=%s"), *FilePath);
-						UE_LOG(LogShaders, Log, TEXT("WaitForSingleObject signaled, Parent process no longer running, exiting"));
-						FPlatformMisc::RequestExit(false);
+						FPlatformMisc::RequestExit(false, TEXT("ShaderCompileWorker - WaitForSingleObject signaled, Parent process no longer running, exiting"));
 					}
 					CloseHandle(ParentProcessHandle);
 				}
@@ -920,7 +927,7 @@ private:
 			if (CurrentTime - GLastCompileTime > TimeToLive)
 			{
 				UE_LOG(LogShaders, Log, TEXT("No jobs found for %f seconds, exiting"), (float)(CurrentTime - GLastCompileTime));
-				FPlatformMisc::RequestExit(false);
+				FPlatformMisc::RequestExit(false, TEXT("ShaderCompileWorker - No jobs found for %f seconds, exiting"));
 			}
 		}
 #endif
