[shaders] reworking ShaderDebugData output to unify the "worker input" and "direct compile" debug mechanisms
- removing r.ShaderCompiler.DebugDumpWorkerInputs cvar; always dump a valid worker input file for all jobs whenever debug info is enabled (further writing of this file has been moved into the FShaderCommonCompileJob::OnComplete implementations alongside the other debug artifacts, so will be written out for cache hits if r.ShaderCompiler.DumpDebugInfoForCacheHits is true)
- skip serializing source when dumping these input files, modify the commandline to use the debug usf files for source instead
- rename "direct compile" mode to "debug compile" and make it use the worker input file instead of needing to serialize extra cruft into the debug .usf file
- remove the r.DumpShaderDebugWorkerCommandLine cvar; all jobs will export a DebugCompileArgs.txt file containing the commandline that can be used to run SCW in debug compile mode by default
- modify pipeline jobs to export a single worker input file (in the first stage's debug folder); each stage of a pipeline will export their own DebugCompileArgs.txt (for convenience, so if only a single stage fails you can still navigate to that stage's debug folder and run the job from there). The debug compile for any pipeline stage will run the full pipeline job (so compilation always matches what happens in the normal context) with paths modified to be relative to that specific debug folder for convenience (so when downloading debug output from a horde build or from another user, only one path in the commandline needs to be fixed up to run on a different machine)

diff --git a/Engine/Source/Developer/ShaderCompilerCommon/Private/ShaderCompilerCommon.cpp b/Engine/Source/Developer/ShaderCompilerCommon/Private/ShaderCompilerCommon.cpp
--- a/Engine/Source/Developer/ShaderCompilerCommon/Private/ShaderCompilerCommon.cpp
+++ b/Engine/Source/Developer/ShaderCompilerCommon/Private/ShaderCompilerCommon.cpp
@@ -1464,59 +1464,6 @@ void CleanupUniformBufferCode(const FShaderCompilerEnvironment& Environment, FSh
 
 ///////////////////////////////////////////////////////////////////////////////////////////////////////////////////////////////////////
 
-FString CreateShaderCompilerWorkerDirectCommandLine(const FShaderCompilerInput& Input, const UE::ShaderCompilerCommon::FDebugShaderDataOptions& Options, const TCHAR* Suffix = nullptr)
-{
-	FString Text(TEXT("-directcompile -format="));
-	Text += Input.ShaderFormat.GetPlainNameString();
-	Text += TEXT(" -entry=\"");
-	Text += Input.EntryPointName;
-
-	Text += TEXT("\" -shaderPlatformName=");
-	Text += Input.ShaderPlatformName.GetPlainNameString();
-
-	Text += FString::Printf(TEXT(" -supportedHardwareMask=%u"), Input.SupportedHardwareMask);
-
-	switch (Input.Target.Frequency)
-	{
-	case SF_Vertex:			Text += TEXT(" -vs"); break;
-	case SF_Mesh:			Text += TEXT(" -ms"); break;
-	case SF_Amplification:	Text += TEXT(" -as"); break;
-	case SF_Geometry:		Text += TEXT(" -gs"); break;
-	case SF_Pixel:			Text += TEXT(" -ps"); break;
-	case SF_Compute:		Text += TEXT(" -cs"); break;
-	case SF_RayGen:			Text += TEXT(" -rgs"); break;
-	case SF_RayMiss:		Text += TEXT(" -rms"); break;
-	case SF_RayHitGroup:	Text += TEXT(" -rhs"); break;
-	case SF_RayCallable:	Text += TEXT(" -rcs"); break;
-	case SF_WorkGraphRoot:			Text += TEXT(" -wrs"); break;
-	case SF_WorkGraphComputeNode:	Text += TEXT(" -wcs"); break;
-	default: break;
-	}
-	if (Input.bCompilingForShaderPipeline)
-	{
-		Text += TEXT(" -pipeline");
-	}
-	if (Input.bIncludeUsedOutputs)
-	{
-		Text += TEXT(" -usedoutputs=");
-		for (int32 Index = 0; Index < Input.UsedOutputs.Num(); ++Index)
-		{
-			if (Index != 0)
-			{
-				Text += TEXT("+");
-			}
-			Text += Input.UsedOutputs[Index];
-		}
-	}
-
-	Text += TEXT(" ");
-	Text += Options.GetDebugShaderPath(Input, Suffix);
-
-	// When we're running in directcompile mode, we don't to spam the crash reporter
-	Text += TEXT(" -nocrashreports");
-	return Text;
-}
-
 static FString CreateShaderConductorCommandLine(const FShaderCompilerInput& Input, const FString& SourceFilename, EShaderConductorTarget SCTarget)
 {
 	const TCHAR* Stage = nullptr;
@@ -2044,27 +1991,6 @@ namespace UE::ShaderCompilerCommon
 	{
 		TRACE_CPUPROFILER_EVENT_SCOPE(FBaseShaderFormat_PreprocessShader);
 
-		if (EnumHasAnyFlags(Input.DebugInfoFlags, EShaderDebugInfoFlags::CompileFromDebugUSF))
-		{
-			// the "VirtualSourceFilePath" given is actually an absolute path to a dumped debug USF file; load it directly.
-			// this occurs when running SCW in "direct compile" mode; this file will already be preprocessed.
-			FString DebugUSF;
-			bool bSuccess = FFileHelper::LoadFileToString(DebugUSF, *Input.VirtualSourceFilePath);
-
-			if (bSuccess)
-			{
-				// const_cast for compile environment; need to populate a subset of environment parameters from parsing comments in the preprocessed code
-				UE::ShaderCompilerCommon::SerializeEnvironmentFromBase64(const_cast<FShaderCompilerEnvironment&>(Input.Environment), DebugUSF);
-
-				// strip comments from source when loading from a debug USF. some backends don't handle the comments that the debug dump inserts properly.
-				TArray<ANSICHAR> Stripped;
-				ShaderConvertAndStripComments(DebugUSF, Stripped);
-				PreprocessOutput.EditSource().Set({ Stripped.GetData(), Stripped.Num() });
-			}
-
-			return bSuccess;
-		}
-
 		check(CheckVirtualShaderFilePath(Input.VirtualSourceFilePath));
 
 		bool bSuccess = ::PreprocessShader(PreprocessOutput, Input, Environment, AdditionalDefines);
@@ -2108,12 +2034,13 @@ namespace UE::ShaderCompilerCommon
 	void FBaseShaderFormat::OutputDebugData(
 		const FShaderCompilerInput& Input,
 		const FShaderPreprocessOutput& PreprocessOutput,
-		const FShaderCompilerOutput& Output) const
+		const FShaderCompilerOutput& Output,
+		FShaderDebugDataContext& Ctx) const
 	{
-		DumpExtendedDebugShaderData(Input, PreprocessOutput, Output);
+		DumpExtendedDebugShaderData(Input, PreprocessOutput, Output, Ctx);
 	}
 
-	static void DumpDebugShaderDataInternal(const FShaderCompilerInput& Input, FStringView PreprocessedSource, const FDebugShaderDataOptions& Options, const TCHAR* Suffix)
+	static void DumpDebugShaderDataInternal(const FShaderCompilerInput& Input, FStringView PreprocessedSource, FShaderDebugDataContext& Ctx, const FDebugShaderDataOptions& Options, const TCHAR* Suffix)
 	{
 		if (!Input.DumpDebugInfoEnabled())
 		{
@@ -2121,37 +2048,42 @@ namespace UE::ShaderCompilerCommon
 		}
 
 		FString Contents = UE::ShaderCompilerCommon::GetDebugShaderContents(Input, PreprocessedSource, Options, Suffix);
-		FFileHelper::SaveStringToFile(Contents, *Options.GetDebugShaderPath(Input, Suffix));
+		FString DebugSourcePath = Options.GetDebugShaderPath(Input, Suffix);
+		FFileHelper::SaveStringToFile(Contents, *DebugSourcePath);
+
+		Ctx.DebugSourceFiles.Add(MoveTemp(DebugSourcePath));
 
-		if (EnumHasAnyFlags(Input.DebugInfoFlags, EShaderDebugInfoFlags::DirectCompileCommandLine) && !Options.bSourceOnly)
-		{
-			FFileHelper::SaveStringToFile(CreateShaderCompilerWorkerDirectCommandLine(Input, Options, Suffix), *GetDebugFileName(Input, Options, TEXT("DirectCompile.txt")));
-		}
 	}
 
 	void DumpDebugShaderData(const FShaderCompilerInput& Input, const FString& PreprocessedSource, const FDebugShaderDataOptions& Options)
 	{
-		DumpDebugShaderDataInternal(Input, PreprocessedSource, Options, nullptr);
+		// deprecated
 	}
 
-
-	void FBaseShaderFormat::OutputDebugDataMinimal(const FShaderCompilerInput& Input) const
+	void FBaseShaderFormat::OutputDebugDataMinimal(const FShaderCompilerInput& Input, FShaderDebugDataContext& Ctx) const
 	{
 		const FStringView FailedSourceStr = TEXTVIEW(R"(
 // Preprocessing failed for shader; defines used can be seen above.
-// DirectCompile cmdline as given at the bottom of the file can be used to inspect the contents of the compilation environment.)");
+// ShaderCompileWorker can be run via the debugger using the commandline args in DebugCompileArgs.txt to inspect the contents of the FShaderCompilerInput/FShaderCompilerEnvironment.)");
+
+		DumpDebugShaderDataInternal(Input, FailedSourceStr, Ctx, FDebugShaderDataOptions(), nullptr);
+	}
 
-		FDebugShaderDataOptions Options;
-		Options.bSourceOnly = true;
-		// Invoking this with bSourceOnly=true will dump only a .usf file; given that we don't have preprocessed source
-		// we use the FailedSourceStr to indicate what value such a file can provide.
-		DumpDebugShaderDataInternal(Input, FailedSourceStr, Options, nullptr);
+
+	void DumpExtendedDebugShaderData(
+		const FShaderCompilerInput& Input,
+		const FShaderPreprocessOutput& PreprocessOutput,
+		const FShaderCompilerOutput& Output,
+		const FDebugShaderDataOptions& Options)
+	{
+		// deprecated
 	}
 
 	void DumpExtendedDebugShaderData(
 		const FShaderCompilerInput& Input,
 		const FShaderPreprocessOutput& PreprocessOutput,
 		const FShaderCompilerOutput& Output,
+		FShaderDebugDataContext& Ctx,
 		const FDebugShaderDataOptions& Options)
 	{
 		if (!Input.Environment.CompilerFlags.Contains(CFLAG_DisableSourceStripping) && EnumHasAnyFlags(Input.DebugInfoFlags, EShaderDebugInfoFlags::DetailedSource))
@@ -2163,15 +2095,15 @@ namespace UE::ShaderCompilerCommon
 		bool bHasModifiedSource = !Output.ModifiedShaderSource.IsEmpty();
 		if (bHasModifiedSource)
 		{
-			// If the compile step applies modifications to the source, output this as the "default" USF; it's not directcompile-compatible but backends
+			// If the compile step applies modifications to the source, output this as the "default" USF; it's not compatible with SCW debug compile mode but backends
 			// which output compile batch files rely on this being the copy of the source that can be passed directly to the platform compiler.
 			FFileHelper::SaveStringToFile(Output.ModifiedShaderSource, *Options.GetDebugShaderPath(Input));
 		}
 
 		// if no modifications to source are made in the compile step, output just the single usf which is the unstripped version compatible with launching
-		// SCW in directcompile mode (the stripped version is less useful for debugging via this mechanism, so is only output in "detailed source" mode)
-		// if modifications were made, this is output as an additional artifact, appending "_DirectCompile" to the path to indicate that it can be used as such.
-		DumpDebugShaderDataInternal(Input, PreprocessOutput.GetUnstrippedSourceView(), Options, bHasModifiedSource ? TEXT("_DirectCompile") : nullptr);
+		// SCW in debug compile mode (the stripped version is less useful for debugging via this mechanism, so is only output in "detailed source" mode)
+		// if modifications were made, this is output as an additional artifact, appending "_DebugCompile" to the path to indicate that it can be used as such.
+		DumpDebugShaderDataInternal(Input, PreprocessOutput.GetUnstrippedSourceView(), Ctx, Options, bHasModifiedSource ? TEXT("_DebugCompile") : nullptr);
 
 		if (EnumHasAnyFlags(Input.DebugInfoFlags, EShaderDebugInfoFlags::ShaderCodePlatformHashes))
 		{
@@ -2254,31 +2186,14 @@ namespace UE::ShaderCompilerCommon
 	
 	FString SerializeEnvironmentToBase64(const FShaderCompilerEnvironment& Env)
 	{
-		TArray<uint8> Serialized;
-		FMemoryWriter Ar(Serialized);
-		const_cast<FShaderCompilerEnvironment&>(Env).SerializeCompilationDependencies(Ar);
-		return FString::Printf(TEXT("%s%s%s"), Base64EnvBegin, *FBase64::Encode(Serialized), Base64EnvEnd);
+		// deprecated
+		static FString Empty;
+		return Empty;
 	}
 
 	void SerializeEnvironmentFromBase64(FShaderCompilerEnvironment& Env, const FString& DebugShaderSource)
 	{
-		int32 BeginIndex = DebugShaderSource.Find(Base64EnvBegin, ESearchCase::CaseSensitive);
-		if (BeginIndex == INDEX_NONE)
-		{
-			return;
-		}
-		int32 EndIndex = DebugShaderSource.Find(Base64EnvEnd, ESearchCase::CaseSensitive, ESearchDir::FromStart, BeginIndex);
-		if (EndIndex == INDEX_NONE)
-		{
-			return;
-		}
-
-		FString Base64Encoded = DebugShaderSource.Left(EndIndex).Mid(BeginIndex + Base64EnvBeginLen);
-
-		TArray<uint8> Serialized;
-		FBase64::Decode(Base64Encoded, Serialized);
-		FMemoryReader Ar(Serialized);
-		Env.SerializeCompilationDependencies(Ar);
+		// deprecated
 	}
 
 	FString GetDebugShaderContents(const FShaderCompilerInput& Input, FStringView PreprocessedSource, const FDebugShaderDataOptions& Options, const TCHAR* Suffix)
@@ -2305,10 +2220,6 @@ namespace UE::ShaderCompilerCommon
 		}
 
 		Contents += TEXT("\n");
-		Contents += SerializeEnvironmentToBase64(MergedEnvironment);
-		Contents += TEXT("/* DIRECT COMPILE\n");
-		Contents += CreateShaderCompilerWorkerDirectCommandLine(Input, Options, Suffix);
-		Contents += TEXT("\nDIRECT COMPILE */\n");
 		if (!Input.DebugDescription.IsEmpty())
 		{
 			Contents += TEXT("//");
diff --git a/Engine/Source/Developer/ShaderCompilerCommon/Public/ShaderCompilerCommon.h b/Engine/Source/Developer/ShaderCompilerCommon/Public/ShaderCompilerCommon.h
--- a/Engine/Source/Developer/ShaderCompilerCommon/Public/ShaderCompilerCommon.h
+++ b/Engine/Source/Developer/ShaderCompilerCommon/Public/ShaderCompilerCommon.h
@@ -114,7 +114,8 @@ namespace UE::ShaderCompilerCommon
 		TFunction<FString()> AppendPreSource{};
 		TFunction<FString()> AppendPostSource{};
 		TArray<FAdditionalOutput> AdditionalOutputs;
-		bool bSourceOnly = false; // if true, will only output source .usf as directed and skip all other debug data artifacts
+		UE_DEPRECATED(5.7, "bSourceOnly flag on debug data options is no longer supported")
+		bool bSourceOnly = false;
 
 		SHADERCOMPILERCOMMON_API FString GetDebugShaderPath(const FShaderCompilerInput& Input, const TCHAR* Suffix = nullptr) const;
 	};
@@ -130,6 +131,13 @@ namespace UE::ShaderCompilerCommon
 	UE_DEPRECATED(5.6, "This function should only be used internally and will be removed from API")
 	extern SHADERCOMPILERCOMMON_API void DumpDebugShaderData(const FShaderCompilerInput& Input, const FString& PreprocessedSource, const FDebugShaderDataOptions& Options = FDebugShaderDataOptions());
 
+	UE_DEPRECATED(5.7, "DumpExtendedDebugShaderData now requires a FShaderDebugDataContext struct parameter")
+		extern SHADERCOMPILERCOMMON_API void DumpExtendedDebugShaderData(
+			const FShaderCompilerInput& Input,
+			const FShaderPreprocessOutput& PreprocessOutput,
+			const FShaderCompilerOutput& Output,
+			const UE::ShaderCompilerCommon::FDebugShaderDataOptions& Options = FDebugShaderDataOptions());
+
 	/*
 	 * Dumps extended debug information; including all outputs from DumpDebugShaderData as well as the following:
 	 *		- OutputHash.txt file containing the SHA hash of the shader job output
@@ -144,15 +152,17 @@ namespace UE::ShaderCompilerCommon
 		const FShaderCompilerInput& Input,
 		const FShaderPreprocessOutput& PreprocessOutput,
 		const FShaderCompilerOutput& Output,
+		FShaderDebugDataContext& Ctx,
 		const UE::ShaderCompilerCommon::FDebugShaderDataOptions& Options = FDebugShaderDataOptions());
 
+	UE_DEPRECATED(5.7, "No longer used")
 	extern SHADERCOMPILERCOMMON_API void SerializeEnvironmentFromBase64(FShaderCompilerEnvironment& Env, const FString& DebugUSF);
+	UE_DEPRECATED(5.7, "No longer used")
 	extern SHADERCOMPILERCOMMON_API FString SerializeEnvironmentToBase64(const FShaderCompilerEnvironment& Env);
 
 	/*
 	 * Constructs the modified preprocessed source that would be dumped to a .usf file via DumpDebugShaderData, including the following additions:
-	 * - resource table debug information
-	 * - commandlines for running ShaderCompileWorker in DirectCompile mode (one which re-runs the preprocessor and one which bypasses it)
+	 * - comments containing the names and values of all defines when preprocessing this shader
 	 * - a comment containing the FShaderCompilerInput DebugDescription
 	 * - any additions performed by the AppendPreSource/AppendPostSource TFunctions
 	 * @param	Input The input of the compilation job
@@ -177,9 +187,10 @@ namespace UE::ShaderCompilerCommon
 		virtual SHADERCOMPILERCOMMON_API void OutputDebugData(
 			const FShaderCompilerInput& Input,
 			const FShaderPreprocessOutput& PreprocessOutput,
-			const FShaderCompilerOutput& Output) const override;
+			const FShaderCompilerOutput& Output,
+			FShaderDebugDataContext& Ctx) const override;
 
-		virtual SHADERCOMPILERCOMMON_API void OutputDebugDataMinimal(const FShaderCompilerInput& Input) const override;
+		virtual SHADERCOMPILERCOMMON_API void OutputDebugDataMinimal(const FShaderCompilerInput& Input, FShaderDebugDataContext& Ctx) const override;
 	};
 }
 
diff --git a/Engine/Source/Developer/TargetPlatform/Public/Interfaces/IShaderFormat.h b/Engine/Source/Developer/TargetPlatform/Public/Interfaces/IShaderFormat.h
--- a/Engine/Source/Developer/TargetPlatform/Public/Interfaces/IShaderFormat.h
+++ b/Engine/Source/Developer/TargetPlatform/Public/Interfaces/IShaderFormat.h
@@ -14,6 +14,7 @@ class FSharedBuffer;
 struct FShaderCompilerEnvironment;
 struct FShaderCompilerInput;
 struct FShaderCompilerOutput;
+struct FShaderDebugDataContext;
 
 /**
  * IShaderFormat, shader pre-compilation abstraction
@@ -186,7 +187,8 @@ public:
 	virtual void OutputDebugData(
 		const FShaderCompilerInput& Input,
 		const FShaderPreprocessOutput& PreprocessOutput,
-		const FShaderCompilerOutput& Output) const {};
+		const FShaderCompilerOutput& Output,
+		FShaderDebugDataContext& Ctx) const {};
 
 	/*
 	 * Implement to output debug info for the case where a secondary preprocessed shader was created. As with the dual-output version
@@ -199,10 +201,11 @@ public:
 		const FShaderPreprocessOutput& PreprocessOutput,
 		const FShaderPreprocessOutput& SecondaryPreprocessOutput,
 		const FShaderCompilerOutput& Output,
-		const FShaderCompilerOutput& SecondaryOutput) const {};
+		const FShaderCompilerOutput& SecondaryOutput,
+		FShaderDebugDataContext& Ctx) const {};
 
 	/* Implement to output a minimal set of debug info based solely on the compile input. This is invoked when preprocessing fails. */
-	virtual void OutputDebugDataMinimal(const FShaderCompilerInput& Input) const {};
+	virtual void OutputDebugDataMinimal(const FShaderCompilerInput& Input, FShaderDebugDataContext& Ctx) const {};
 
 public:
 	/** Virtual destructor. */
diff --git a/Engine/Source/Programs/ShaderCompileWorker/Private/ShaderCompileWorker.cpp b/Engine/Source/Programs/ShaderCompileWorker/Private/ShaderCompileWorker.cpp
--- a/Engine/Source/Programs/ShaderCompileWorker/Private/ShaderCompileWorker.cpp
+++ b/Engine/Source/Programs/ShaderCompileWorker/Private/ShaderCompileWorker.cpp
@@ -230,8 +230,9 @@ public:
 	int32 NumberToProcess = -1;
 	bool DisableFileWrite = false;
 	bool KeepInput = false;
+	TArray<FString> DebugSourceFiles;
 
-	FWorkLoop(const TCHAR* ParentProcessIdText,const TCHAR* InWorkingDirectory,const TCHAR* InInputFilename,const TCHAR* InOutputFilename, TMap<FString, uint32>& InFormatVersionMap)
+	FWorkLoop(const TCHAR* ParentProcessIdText,const TCHAR* InWorkingDirectory,const TCHAR* InInputFilename,const TCHAR* InOutputFilename, TMap<FName, uint32>& InFormatVersionMap)
 	:	ParentProcessId(FCString::Atoi(ParentProcessIdText))
 	,	WorkingDirectory(InWorkingDirectory)
 	,	InputFilename(InInputFilename)
@@ -260,6 +261,16 @@ public:
 			{
 				NumberToProcess = FCString::Atoi(Switch.GetCharArray().GetData() + 8);
 			}
+			else if (Switch.StartsWith(TEXT("debugsourcefiles=")))
+			{
+				FString DelimitedSourceFiles = Switch.RightChop(17);
+				DelimitedSourceFiles.ParseIntoArray(DebugSourceFiles, TEXT(","));
+				for (FString& DebugSourceFile : DebugSourceFiles)
+				{
+					// debug source files are given as relative paths from working directory, make them absolute here and normalize
+					DebugSourceFile = InWorkingDirectory / DebugSourceFile;
+				}
+			}
 		}
 	}
 
@@ -433,7 +444,7 @@ private:
 
 	FString InputFilePath;
 	FString OutputFilePath;
-	TMap<FString, uint32> FormatVersionMap;
+	TMap<FName, uint32> FormatVersionMap;
 	FString TempFilePath;
 
 	/** Opens an input file, trying multiple times if necessary. */
@@ -457,7 +468,7 @@ private:
 		return InputFile;
 	}
 
-	void VerifyFormatVersions(TMap<FString, uint32>& ReceivedFormatVersionMap)
+	void VerifyFormatVersions(TMap<FName, uint32>& ReceivedFormatVersionMap)
 	{
 		for (auto Pair : ReceivedFormatVersionMap)
 		{
@@ -466,7 +477,7 @@ private:
 			{
 				if (Pair.Value != *Found)
 				{
-					ExitWithoutCrash(FSCWErrorCode::BadShaderFormatVersion, FString::Printf(TEXT("Mismatched shader version for format %s: Found version %u but expected %u; did you forget to build ShaderCompilerWorker?"), *Pair.Key, *Found, Pair.Value));
+					ExitWithoutCrash(FSCWErrorCode::BadShaderFormatVersion, FString::Printf(TEXT("Mismatched shader version for format %s: Found version %u but expected %u; did you forget to build ShaderCompilerWorker?"), *Pair.Key.ToString(), *Found, Pair.Value));
 				}
 			}
 		}
@@ -515,10 +526,18 @@ private:
 		FMemoryReader InputMemory(UncompressedData);
 		FArchive& InputFile = bWasCompressed ? InputMemory : *InputFilePtr;
 
-		TMap<FString, uint32> ReceivedFormatVersionMap;
+		TMap<FName, uint32> ReceivedFormatVersionMap;
 		InputFile << ReceivedFormatVersionMap;
 
-		VerifyFormatVersions(ReceivedFormatVersionMap);
+		bool bIsDebugCompile = !DebugSourceFiles.IsEmpty();
+
+		// when running a debug compile we are more permissive w.r.t. format versions.
+		// (instead of requiring all versions match, we just check the version for the debug compilation
+		// that we will execute)
+		if (!bIsDebugCompile)
+		{
+			VerifyFormatVersions(ReceivedFormatVersionMap);
+		}
 		
 		// Initialize shader hash cache before reading any includes.
 		InitializeShaderHashCache();
@@ -684,6 +703,30 @@ private:
 
 		GNumProcessedJobs = 0;
 
+		auto SetupDebugCompile = [this, &ReceivedFormatVersionMap](FShaderCompileJob& Job, const TCHAR* SourcePath)
+			{
+				Job.Input.DebugInfoFlags &= EShaderDebugInfoFlags::CompileFromDebugUSF;
+
+				checkf(ReceivedFormatVersionMap[Job.Input.ShaderFormat] == FormatVersionMap[Job.Input.ShaderFormat],
+						TEXT("Format mismatch for shader format %s; ensure SDK/tools for this platform match the build where this debug compile input was generated.\n")
+						TEXT("This check is skippable, with the caveat that this debug compilation may fail for unexpected reasons and/or results may be different."),
+						*Job.Input.ShaderFormat.ToString());
+
+				Job.Input.DebugInfoFlags &= EShaderDebugInfoFlags::CompileFromDebugUSF;
+				FString DebugSource;
+				bool bSuccess = FFileHelper::LoadFileToString(DebugSource, SourcePath);
+
+				if (bSuccess)
+				{
+					// strip comments from source when loading from a debug source file. some backends don't handle the comments that the debug dump inserts properly.
+					FAnsiString Stripped;
+					// Since we are directly assigning the stripped result to an FShaderSource object which automatically appends SIMD patting as part of the Set function,
+					// we don't need to also add padding during the convert-and-strip operation.
+					ShaderConvertAndStripComments(DebugSource, Stripped.GetCharArray(), EConvertAndStripFlags::NoSimdPadding);
+					Job.PreprocessOutput.EditSource() = MoveTemp(Stripped);
+				}
+			};
+
 		// Individual jobs
 		{
 			int32 SingleJobHeader = ShaderCompileWorkerSingleJobHeader;
@@ -693,22 +736,35 @@ private:
 				ExitWithoutCrash(FSCWErrorCode::BadSingleJobHeader, FString::Printf(TEXT("Exiting due to ShaderCompilerWorker expecting job header %d, got %d instead! Did you forget to build ShaderCompilerWorker?"), ShaderCompileWorkerSingleJobHeader, SingleJobHeader));
 			}
 
-			int32 NumBatches = 0;
-			InputFile << NumBatches;
+			int32 NumSingleJobs = 0;
+			InputFile << NumSingleJobs;
 
 			FlushShaderFileCache();
 
-			OutSingleJobs.Reserve(NumBatches);
-			
-			for (int32 BatchIndex = 0; BatchIndex < NumBatches; BatchIndex++)
+			OutSingleJobs.Reserve(NumSingleJobs);
+		
+			for (int32 BatchIndex = 0; BatchIndex < NumSingleJobs; BatchIndex++)
 			{
 				FShaderCompileJob& Job = OutSingleJobs.AddDefaulted_GetRef();
 				// Deserialize the job's inputs.
 				Job.SerializeWorkerInput(InputFile);
 				Job.Input.DeserializeSharedInputs(InputFile, SharedEnvironments, ParameterStructures);
+				if (bIsDebugCompile)
+				{
+
+					// this should only be used in debug compile mode; in this case we expect either a single job or a single pipeline job
+					check(DebugSourceFiles.Num() == 1 && NumSingleJobs == 1);
+					SetupDebugCompile(Job, *DebugSourceFiles[0]);
+				}
 
 				// Process the job.
-				CompileShader(GetShaderFormats(), Job, WorkingDirectory, &GNumProcessedJobs); 
+				CompileShader(GetShaderFormats(), Job, WorkingDirectory, &GNumProcessedJobs);
+				if (bIsDebugCompile)
+				{
+					// if we're in debug compile mode and we have a single job, skip the pipeline job section below so
+					// we can check that in either debug compile case we have the correct number of debug source files.
+					return;
+				}
 			}
 		}
 
@@ -727,15 +783,21 @@ private:
 			OutPipelineNames.Reserve(NumPipelines);
 			OutPipelineJobs.Reserve(NumPipelines);
 
+			// ensure that we are only compiling a single pipeline if debug source files were specified
+			check(!bIsDebugCompile || NumPipelines == 1);
+
 			for (int32 Index = 0; Index < NumPipelines; ++Index)
 			{
 				FString& PipelineName = OutPipelineNames.AddDefaulted_GetRef();
 				InputFile << PipelineName;
-				
+
 				int32 NumStages = 0;
 				InputFile << NumStages;
 				FShaderPipelineCompileJob& PipelineJob = OutPipelineJobs.Emplace_GetRef(NumStages);
 
+				// if debug source files are specified ensure the number of files matches the number of pipeline stages
+				check(!bIsDebugCompile || DebugSourceFiles.Num() == NumStages);
+
 				for (int32 StageIndex = 0; StageIndex < NumStages; ++StageIndex)
 				{
 					// Deserialize the job's inputs.
@@ -743,6 +805,11 @@ private:
 					Job->SerializeWorkerInput(InputFile);
 					Job->Input.DeserializeSharedInputs(InputFile, SharedEnvironments, ParameterStructures);
 
+					if (bIsDebugCompile)
+					{
+						SetupDebugCompile(*Job, *DebugSourceFiles[StageIndex]);
+					}
+
 					// SCW doesn't run DDPI, GShaderHasCache Initialize is run  at start with no knowledge of the CustomPlatforms
 					// CustomPlatforms are known when we parse the WorkerInput so we populate the Directory here
 					if (IsCustomPlatform((EShaderPlatform)Job->Input.Target.Platform))
@@ -935,151 +1002,6 @@ private:
 	}
 };
 
-static void DirectCompile(const TArray<const class IShaderFormat*>& ShaderFormats)
-{
-	TRACE_CPUPROFILER_EVENT_SCOPE(DirectCompile);
-
-	// Find all the info required for compiling a single shader
-	TArray<FString> Tokens, Switches;
-	FCommandLine::Parse(FCommandLine::Get(), Tokens, Switches);
-
-	FString InputFile;
-
-	FName FormatName;
-	FName ShaderPlatformName;
-	FString Entry = TEXT("Main");
-	uint32 SupportedHardwareMask = 0;
-	FString DumpDebugInfoPath;
-	bool bPipeline = false;
-	EShaderFrequency Frequency = SF_Pixel;
-	TArray<FString> UsedOutputs;
-	bool bIncludeUsedOutputs = false;
-	uint64 CFlags = 0;
-	for (const FString& Token : Tokens)
-	{
-		if (Switches.Contains(Token))
-		{
-			if (Token.StartsWith(TEXT("format=")))
-			{
-				FormatName = FName(*Token.RightChop(7));
-			}
-			else if (Token.StartsWith(TEXT("entry=")))
-			{
-				Entry = Token.RightChop(6);
-
-				// Remove quotations marks at beginning and end; happens when multiple entry points are specified, e.g. -entry="closesthit=A anyhit=B"
-				if (Entry.Len() >= 2 && Entry[0] == TEXT('\"') && Entry[Entry.Len() - 1] == TEXT('\"'))
-				{
-					Entry = Entry.Mid(1, Entry.Len() - 2);
-				}
-			}
-			else if (Token.StartsWith(TEXT("shaderPlatformName=")))
-			{
-				ShaderPlatformName = FName(*Token.RightChop(19));
-			}
-			else if (Token.StartsWith(TEXT("supportedHardwareMask=")))
-			{
-				SupportedHardwareMask = (uint32)FCString::Atoi64(*Token.RightChop(22));
-			}
-			else if (Token.StartsWith(TEXT("DebugInfoPath=")))
-			{
-				DumpDebugInfoPath = Token.RightChop(14);
-			}
-			else if (!FCString::Strcmp(*Token, TEXT("ps")))
-			{
-				Frequency = SF_Pixel;
-			}
-			else if (!FCString::Strcmp(*Token, TEXT("vs")))
-			{
-				Frequency = SF_Vertex;
-			}
-			else if (!FCString::Strcmp(*Token, TEXT("ms")))
-			{
-				Frequency = SF_Mesh;
-			}
-			else if (!FCString::Strcmp(*Token, TEXT("as")))
-			{
-				Frequency = SF_Amplification;
-			}
-			else if (!FCString::Strcmp(*Token, TEXT("gs")))
-			{
-				Frequency = SF_Geometry;
-			}
-			else if (!FCString::Strcmp(*Token, TEXT("cs")))
-			{
-				Frequency = SF_Compute;
-			}
-			else if (!FCString::Strcmp(*Token, TEXT("rgs")))
-			{
-				Frequency = SF_RayGen;
-			}
-			else if (!FCString::Strcmp(*Token, TEXT("rms")))
-			{
-				Frequency = SF_RayMiss;
-			}
-			else if (!FCString::Strcmp(*Token, TEXT("rhs")))
-			{
-				Frequency = SF_RayHitGroup;
-			}
-			else if (!FCString::Strcmp(*Token, TEXT("rcs")))
-			{
-				Frequency = SF_RayCallable;
-			}
-			else if (!FCString::Strcmp(*Token, TEXT("wrs")))
-			{
-				Frequency = SF_WorkGraphRoot;
-			}
-			else if (!FCString::Strcmp(*Token, TEXT("wcs")))
-			{
-				Frequency = SF_WorkGraphComputeNode;
-			}
-			else if (!FCString::Strcmp(*Token, TEXT("pipeline")))
-			{
-				bPipeline = true;
-			}
-			else if (Token.StartsWith(TEXT("usedoutputs=")))
-			{
-				FString Outputs = Token.RightChop(12);
-				bIncludeUsedOutputs = true;
-				FString LHS, RHS;
-				while (Outputs.Split(TEXT("+"), &LHS, &RHS))
-				{
-					Outputs = RHS;
-					UsedOutputs.Add(LHS);
-				}
-				UsedOutputs.Add(Outputs);
-			}
-		}
-		else
-		{
-			if (InputFile.Len() == 0)
-			{
-				InputFile = Token;
-			}
-		}
-	}
-
-	FString Dir = FPlatformProcess::UserTempDir();
-
-	FShaderCompileJob Job;
-	Job.Input.EntryPointName = Entry;
-	Job.Input.ShaderFormat = FormatName;
-	Job.Input.ShaderPlatformName = ShaderPlatformName;
-	Job.Input.SupportedHardwareMask = SupportedHardwareMask;
-	Job.Input.VirtualSourceFilePath = InputFile;
-	Job.Input.Target.Platform =  ShaderFormatNameToShaderPlatform(FormatName);
-	Job.Input.Target.Frequency = Frequency;
-	Job.Input.DumpDebugInfoPath = DumpDebugInfoPath;
-	Job.Input.bCompilingForShaderPipeline = bPipeline;
-	Job.Input.bIncludeUsedOutputs = bIncludeUsedOutputs;
-	Job.Input.UsedOutputs = UsedOutputs;
-	Job.Input.DebugInfoFlags = EShaderDebugInfoFlags::CompileFromDebugUSF;
-
-	FShaderCompilerOutput Output;
-	CompileShader(GetShaderFormats(), Job, Dir, &GNumProcessedJobs);
-}
-
-
 /** 
  * Main entrypoint, guarded by a try ... except.
  * This expects 4 parameters:
@@ -1088,7 +1010,7 @@ static void DirectCompile(const TArray<const class IShaderFormat*>& ShaderFormat
  *		The parent process Id
  *		The thread Id corresponding to this worker
  */
-static int32 GuardedMain(int32 argc, TCHAR* argv[], FString& CrashOutputFile, bool bDirectMode)
+static int32 GuardedMain(int32 argc, TCHAR* argv[], FString& CrashOutputFile)
 {
 	FString ExtraCmdLine = TEXT("-NOPACKAGECACHE -ReduceThreadUsage -cpuprofilertrace -nocrashreports");
 
@@ -1194,7 +1116,7 @@ static int32 GuardedMain(int32 argc, TCHAR* argv[], FString& CrashOutputFile, bo
 	// We just enumerate the shader formats here for debugging.
 	const TArray<const class IShaderFormat*>& ShaderFormats = GetShaderFormats();
 	check(ShaderFormats.Num());
-	TMap<FString, uint32> FormatVersionMap;
+	TMap<FName, uint32> FormatVersionMap;
 	for (int32 Index = 0; Index < ShaderFormats.Num(); Index++)
 	{
 		TRACE_CPUPROFILER_EVENT_SCOPE(ShaderFormat);
@@ -1206,34 +1128,28 @@ static int32 GuardedMain(int32 argc, TCHAR* argv[], FString& CrashOutputFile, bo
 		{
 			UE_LOG(LogShaders, Display, TEXT("Available Shader Format %s"), *OutFormats[InnerIndex].ToString());
 			uint32 Version = ShaderFormats[Index]->GetVersion(OutFormats[InnerIndex]);
-			FormatVersionMap.Add(OutFormats[InnerIndex].ToString(), Version);
+			FormatVersionMap.Add(OutFormats[InnerIndex], Version);
 		}
 	}
 
 	GLastCompileTime = FPlatformTime::Seconds();
 
-	if (bDirectMode)
-	{
-		DirectCompile(ShaderFormats);
-	}
-	else
-	{
 #if PLATFORM_WINDOWS
-		//@todo - would be nice to change application name or description to have the ThreadId in it for debugging purposes
-		SetConsoleTitle(argv[3]);
+	//@todo - would be nice to change application name or description to have the ThreadId in it for debugging purposes
+	SetConsoleTitle(argv[3]);
 #endif
 
-		TRACE_CPUPROFILER_EVENT_SCOPE(FWorkLoop);
+	TRACE_CPUPROFILER_EVENT_SCOPE(FWorkLoop);
 
-		FWorkLoop WorkLoop(argv[2], argv[1], argv[4], argv[5], FormatVersionMap);
-		WorkLoop.Loop(CrashOutputFile);
-	}
+	FWorkLoop WorkLoop(argv[2], argv[1], argv[4], argv[5], FormatVersionMap);
+	WorkLoop.Loop(CrashOutputFile);
+	
 
 	return 0;
 }
 
 
-static int32 GuardedMainWrapper(int32 ArgC, TCHAR* ArgV[], FString& CrashOutputFile, bool bDirectMode)
+static int32 GuardedMainWrapper(int32 ArgC, TCHAR* ArgV[], FString& CrashOutputFile)
 {
 	FTaskTagScope Scope(ETaskTag::EGameThread);
 	// We need to know whether we are using XGE now, in case an exception
@@ -1256,7 +1172,7 @@ static int32 GuardedMainWrapper(int32 ArgC, TCHAR* ArgV[], FString& CrashOutputF
 	if (FPlatformMisc::IsDebuggerPresent())
 #endif
 	{
-		ReturnCode = GuardedMain(ArgC, ArgV, CrashOutputFile, bDirectMode);
+		ReturnCode = GuardedMain(ArgC, ArgV, CrashOutputFile);
 	}
 #if PLATFORM_WINDOWS
 	else
@@ -1268,7 +1184,7 @@ static int32 GuardedMainWrapper(int32 ArgC, TCHAR* ArgV[], FString& CrashOutputF
 		__try
 		{
 			GIsGuarded = 1;
-			ReturnCode = GuardedMain(ArgC, ArgV, CrashOutputFile, bDirectMode);
+			ReturnCode = GuardedMain(ArgC, ArgV, CrashOutputFile);
 			GIsGuarded = 0;
 		}
 		__except(HandleShaderCompileException(GetExceptionInformation(), ExceptionMsg, ExceptionCallStack))
@@ -1348,31 +1264,18 @@ INT32_MAIN_INT32_ARGC_TCHAR_ARGV()
 
 	FString OutputFilePath;
 
-	bool bDirectMode = false;
-	for (int32 Index = 1; Index < ArgC; ++Index)
+	if (ArgC < 6)
 	{
-		if (FCString::Strcmp(ArgV[Index], TEXT("-directcompile")) == 0)
-		{
-			bDirectMode = true;
-			break;
-		}
+		printf("ShaderCompileWorker (v%d) is called by UnrealEditor, it requires specific command line arguments.\n", ShaderCompileWorkerOutputVersion);
+		return -1;
 	}
 
-	if (!bDirectMode)
-	{
-		if (ArgC < 6)
-		{
-			printf("ShaderCompileWorker (v%d) is called by UnrealEditor, it requires specific command line arguments.\n", ShaderCompileWorkerOutputVersion);
-			return -1;
-		}
-
-		// Game exe can pass any number of parameters through with appGetSubprocessCommandline
-		// so just make sure we have at least the minimum number of parameters.
-		check(ArgC >= 6);
+	// Game exe can pass any number of parameters through with appGetSubprocessCommandline
+	// so just make sure we have at least the minimum number of parameters.
+	check(ArgC >= 6);
 
-		OutputFilePath = ArgV[1];
-		OutputFilePath += ArgV[5];
-	}
+	OutputFilePath = ArgV[1];
+	OutputFilePath += ArgV[5];
 
-	return GuardedMainWrapper(ArgC, ArgV, OutputFilePath, bDirectMode);
+	return GuardedMainWrapper(ArgC, ArgV, OutputFilePath);
 }
diff --git a/Engine/Source/Runtime/Engine/Private/ShaderCompiler/ShaderCompiler.cpp b/Engine/Source/Runtime/Engine/Private/ShaderCompiler/ShaderCompiler.cpp
--- a/Engine/Source/Runtime/Engine/Private/ShaderCompiler/ShaderCompiler.cpp
+++ b/Engine/Source/Runtime/Engine/Private/ShaderCompiler/ShaderCompiler.cpp
@@ -329,14 +329,6 @@ static FAutoConsoleVariableRef CVarDumpShaderDebugBindlessNames(
 	TEXT("When set to 1, will add bindless folder names.")
 	);
 
-static int32 GDumpShaderDebugInfoSCWCommandLine = 0;
-static FAutoConsoleVariableRef CVarDumpShaderDebugSCWCommandLine(
-	TEXT("r.DumpShaderDebugWorkerCommandLine"),
-	GDumpShaderDebugInfoSCWCommandLine,
-	TEXT("Only valid when r.DumpShaderDebugInfo > 0.\n")
-	TEXT("When set to 1, it will generate a file that can be used with ShaderCompileWorker's -directcompile.")
-	);
-
 static int32 GShaderMapCompilationTimeout = 2 * 60 * 60;	// anything below an hour can hit a false positive
 static FAutoConsoleVariableRef CVarShaderMapCompilationTimeout(
 	TEXT("r.ShaderCompiler.ShadermapCompilationTimeout"),
@@ -1344,10 +1336,6 @@ FShaderCompilingManager::EDumpShaderDebugInfo FShaderCompilingManager::GetDumpSh
 EShaderDebugInfoFlags FShaderCompilingManager::GetDumpShaderDebugInfoFlags() const
 {
 	EShaderDebugInfoFlags Flags = EShaderDebugInfoFlags::Default;
-	if (GDumpShaderDebugInfoSCWCommandLine)
-	{
-		Flags |= EShaderDebugInfoFlags::DirectCompileCommandLine;
-	}
 	
 	if (CVarDebugDumpJobInputHashes.GetValueOnAnyThread())
 	{
@@ -2487,23 +2475,26 @@ bool FShaderCompilingManager::HandlePotentialRetry(TMap<int32, FShaderMapFinaliz
 		for (int32 FinishedJobIndex = 0; FinishedJobIndex < Results.FinishedJobs.Num(); ++FinishedJobIndex)
 		{
 			FShaderCommonCompileJobPtr Job = Results.FinishedJobs[FinishedJobIndex];
-			bool bReissue = false;
-			Job->ForEachSingleShaderJob([&bReissue, bRetryCompile](FShaderCompileJob& SingleJob)
+			bool bReissue = false, bEnableDebug = false;
+			Job->ForEachSingleShaderJob([&bReissue, &bEnableDebug, bRetryCompile](FShaderCompileJob& SingleJob)
 				{
-					// this will set bReissue to true if any stage of a pipeline job failed, so the entire pipeline job will be 
-					// resubmitted, but only the failed stage(s) will have debug info enabled on the re-run.
-					bool bShouldRecompileForDebug = GShaderCompilingManager->ShouldRecompileToDumpShaderDebugInfo(SingleJob);
 					// reissue the job if either we want to dump debug info for it, or a retry was requested above for any failed jobs
-					bReissue = bShouldRecompileForDebug || (bRetryCompile && !SingleJob.bSucceeded);
-					if (bShouldRecompileForDebug)
-					{
-						SingleJob.Input.DumpDebugInfoPath = GShaderCompilingManager->CreateShaderDebugInfoPath(SingleJob.Input);
-						// Any reissued jobs due to this condition will dump debug information, so increment the dump count here
-						GShaderCompilingManager->IncrementNumDumpedShaderSources();
-					}
+					bEnableDebug = GShaderCompilingManager->ShouldRecompileToDumpShaderDebugInfo(SingleJob);
+					bReissue = bEnableDebug || (bRetryCompile && !SingleJob.bSucceeded);
 				});
 			if (bReissue)
 			{
+				if (bEnableDebug)
+				{
+					// if any stage of a pipeline job is reissued to dump debug info, we need to enable debug info on all stages of the job
+					// since the debug information relies on debug source files existing for every stage
+					Job->ForEachSingleShaderJob([](FShaderCompileJob& SingleJob)
+						{
+							SingleJob.Input.DumpDebugInfoPath = GShaderCompilingManager->CreateShaderDebugInfoPath(SingleJob.Input);
+							// Any reissued jobs due to this condition will dump debug information, so increment the dump count here
+							GShaderCompilingManager->IncrementNumDumpedShaderSources();
+						});
+				}
 				bAnyReissued = true;
 				ReissueJobs.Add(Job);
 				// remove the job we will reissue from the finished jobs list; no point in allowing shrinking as it will grow again below when the reissued jobs finish
diff --git a/Engine/Source/Runtime/Engine/Private/ShaderCompiler/ShaderCompilerJobCache.cpp b/Engine/Source/Runtime/Engine/Private/ShaderCompiler/ShaderCompilerJobCache.cpp
--- a/Engine/Source/Runtime/Engine/Private/ShaderCompiler/ShaderCompilerJobCache.cpp
+++ b/Engine/Source/Runtime/Engine/Private/ShaderCompiler/ShaderCompilerJobCache.cpp
@@ -1563,20 +1563,15 @@ double FShaderJobCache::ProcessFinishedJob(FShaderCommonCompileJob* FinishedJob,
 
 	FinishedJob->UpdateStatus(Status);
 
-	bool bCompilationSkipped =
-		FinishedJob->JobStatusPtr->GetIsDuplicate() ||
-		(Status == EShaderCompileJobStatus::Skipped) ||
-		(Status == EShaderCompileJobStatus::CompleteFoundInCache) ||
-		(Status == EShaderCompileJobStatus::CompleteFoundInDDC);
-								
-	if (!bCompilationSkipped)
+	if (!FinishedJob->JobStatusPtr->WasCompilationSkipped())
 	{
 		AddToCacheAndProcessPending(FinishedJob);
 	}
 
-	FinishedJob->OnComplete();
+	FShaderDebugDataContext Ctx;
+	FinishedJob->OnComplete(Ctx);
 
-	GShaderCompilerStats->RegisterFinishedJob(*FinishedJob, bCompilationSkipped);
+	GShaderCompilerStats->RegisterFinishedJob(*FinishedJob);
 
 	{
 		// Need to protect writes to FShaderMapCompileResults
diff --git a/Engine/Source/Runtime/Engine/Private/ShaderCompiler/ShaderCompilerStats.cpp b/Engine/Source/Runtime/Engine/Private/ShaderCompiler/ShaderCompilerStats.cpp
--- a/Engine/Source/Runtime/Engine/Private/ShaderCompiler/ShaderCompilerStats.cpp
+++ b/Engine/Source/Runtime/Engine/Private/ShaderCompiler/ShaderCompilerStats.cpp
@@ -18,6 +18,7 @@
 #include "ProfilingDebugging/LoadTimeTracker.h"
 #include "ProfilingDebugging/StallDetector.h"
 #include "Serialization/CompactBinaryWriter.h"
+#include "ShaderCompileWorkerUtil.h"
 #include "Misc/FileHelper.h"
 
 static int32 GLogShaderCompilerStats = 0;
@@ -911,10 +912,11 @@ void FShaderCompilerStats::RegisterAssignedJob(FShaderCommonCompileJob& Job)
 }
 
 TRACE_DECLARE_INT_COUNTER(Shaders_Compiled, TEXT("Shaders/Compiled"));
-void FShaderCompilerStats::RegisterFinishedJob(FShaderCommonCompileJob& Job, bool bCompilationSkipped)
+void FShaderCompilerStats::RegisterFinishedJob(FShaderCommonCompileJob& Job)
 {
 	FScopeLock Lock(&CompileStatsLock);
-	
+
+	bool bCompilationSkipped = Job.JobStatusPtr->WasCompilationSkipped();
 	if (!bCompilationSkipped)
 	{
 		ensure(Job.TimeAssignedToExecution != 0.0);
diff --git a/Engine/Source/Runtime/Engine/Private/ShaderCompiler/ShaderCompilerThreadRunnable.cpp b/Engine/Source/Runtime/Engine/Private/ShaderCompiler/ShaderCompilerThreadRunnable.cpp
--- a/Engine/Source/Runtime/Engine/Private/ShaderCompiler/ShaderCompilerThreadRunnable.cpp
+++ b/Engine/Source/Runtime/Engine/Private/ShaderCompiler/ShaderCompilerThreadRunnable.cpp
@@ -6,7 +6,6 @@
 =============================================================================*/
 
 #include "ShaderCompilerPrivate.h"
-#include "ShaderCompilerMemoryLimit.h"
 
 #include "Async/ParallelFor.h"
 #include "HAL/PlatformFile.h"
@@ -18,6 +17,7 @@
 #include "ProfilingDebugging/LoadTimeTracker.h"
 #include "ProfilingDebugging/ScopedTimers.h"
 #include "ProfilingDebugging/StallDetector.h"
+#include "ShaderCompilerMemoryLimit.h"
 #include "ShaderCompileWorkerUtil.h"
 
 
@@ -836,8 +836,6 @@ void FShaderCompileThreadRunnable::CompileDirectlyThroughDll()
 
 		if (CurrentWorkerInfo.QueuedJobs.Num() > 0)
 		{
-			FShaderCompileWorkerUtil::DumpDebugCompileInput(CurrentWorkerInfo.QueuedJobs);
-
 			for (int32 JobIndex = 0; JobIndex < CurrentWorkerInfo.QueuedJobs.Num(); JobIndex++)
 			{
 				FShaderCommonCompileJob& CurrentJob = *CurrentWorkerInfo.QueuedJobs[JobIndex];
diff --git a/Engine/Source/Runtime/Engine/Public/ShaderCompiler.h b/Engine/Source/Runtime/Engine/Public/ShaderCompiler.h
--- a/Engine/Source/Runtime/Engine/Public/ShaderCompiler.h
+++ b/Engine/Source/Runtime/Engine/Public/ShaderCompiler.h
@@ -831,7 +831,7 @@ public:
 	void RegisterAssignedJob(FShaderCommonCompileJob& InOutJob);
 
 	/** Marks the job as finished for the stats purpose. Job will be modified to include the current timestamp. */
-	void RegisterFinishedJob(FShaderCommonCompileJob& InOutJob, bool bCompilationSkipped);
+	void RegisterFinishedJob(FShaderCommonCompileJob& InOutJob);
 
 	/** Informs statistics about a new job batch, so we can tally up batches. */
 	void RegisterJobBatch(int32 NumJobs, EExecutionType ExecType);
diff --git a/Engine/Source/Runtime/RenderCore/Internal/ShaderCompileWorkerUtil.h b/Engine/Source/Runtime/RenderCore/Internal/ShaderCompileWorkerUtil.h
--- a/Engine/Source/Runtime/RenderCore/Internal/ShaderCompileWorkerUtil.h
+++ b/Engine/Source/Runtime/RenderCore/Internal/ShaderCompileWorkerUtil.h
@@ -54,6 +54,7 @@ public:
 	{
 		None = 0,
 		CompressTaskFile = 1 << 0,
+		SkipSource = 1 << 1,
 	};
 
 	enum class EReadTasksFlags : uint8
@@ -63,7 +64,7 @@ public:
 	};
 
 	RENDERCORE_API static void LogQueuedCompileJobs(const TArray<FShaderCommonCompileJobPtr>& QueuedJobs, int32 NumProcessedJobs);
-	RENDERCORE_API static void DumpDebugCompileInput(TConstArrayView<FShaderCommonCompileJobPtr> QueuedJobs);
+	RENDERCORE_API static void DumpDebugCompileInput(FShaderCommonCompileJob& Job, FShaderDebugDataContext& Ctx);
 	RENDERCORE_API static bool WriteTasks(const TArray<FShaderCommonCompileJobPtr>& QueuedJobs, FArchive& TransferFile, EWriteTasksFlags Flags = EWriteTasksFlags::None);
 	RENDERCORE_API static FSCWErrorCode::ECode ReadTasks(const TArray<FShaderCommonCompileJobPtr>& QueuedJobs, FArchive& OutputFile, FShaderCompileWorkerDiagnostics* OutWorkerDiagnostics, EReadTasksFlags Flags = EReadTasksFlags::None);
 
diff --git a/Engine/Source/Runtime/RenderCore/Private/ShaderCompileWorkerUtil.cpp b/Engine/Source/Runtime/RenderCore/Private/ShaderCompileWorkerUtil.cpp
--- a/Engine/Source/Runtime/RenderCore/Private/ShaderCompileWorkerUtil.cpp
+++ b/Engine/Source/Runtime/RenderCore/Private/ShaderCompileWorkerUtil.cpp
@@ -11,12 +11,6 @@
 #include "Interfaces/ITargetPlatformManagerModule.h"
 #include "Interfaces/IShaderFormat.h"
 
-static TAutoConsoleVariable<bool> CVarDebugDumpWorkerInputs(
-	TEXT("r.ShaderCompiler.DebugDumpWorkerInputs"),
-	false,
-	TEXT("If true, worker input files will be saved for each individual compile job alongside other debug data (note that r.DumpShaderDebugInfo must also be enabled for this to function)"),
-	ECVF_ReadOnly);
-
 static TAutoConsoleVariable<bool> CVarShadersPropagateLocalWorkerOOMs(
 	TEXT("r.Shaders.PropagateLocalWorkerOOMs"),
 	false,
@@ -68,9 +62,9 @@ static void ModalErrorOrLog(const FString& Title, const FString& Text, int64 Cur
 	}
 }
 
-static TMap<FString, uint32> GetFormatVersionMap()
+static TMap<FName, uint32> GetFormatVersionMap()
 {
-	TMap<FString, uint32> FormatVersionMap;
+	TMap<FName, uint32> FormatVersionMap;
 
 	const TArray<const class IShaderFormat*>& ShaderFormats = GetTargetPlatformManagerRef().GetShaderFormats();
 	check(ShaderFormats.Num());
@@ -82,7 +76,7 @@ static TMap<FString, uint32> GetFormatVersionMap()
 		for (int32 InnerIndex = 0; InnerIndex < OutFormats.Num(); InnerIndex++)
 		{
 			uint32 Version = ShaderFormats[Index]->GetVersion(OutFormats[InnerIndex]);
-			FormatVersionMap.Add(OutFormats[InnerIndex].ToString(), Version);
+			FormatVersionMap.Add(OutFormats[InnerIndex], Version);
 		}
 	}
 
@@ -401,7 +395,8 @@ static void SplitJobsByType(const TArray<FShaderCommonCompileJobPtr>& QueuedJobs
 	}
 }
 
-static bool DoWriteTasksInner(const TArray<FShaderCommonCompileJobPtr>& QueuedJobs, FArchive& InTransferFile, FShaderCompileWorkerUtil::EWriteTasksFlags Flags)
+// Serialize Queued Job information
+bool FShaderCompileWorkerUtil::WriteTasks(const TArray<FShaderCommonCompileJobPtr>& QueuedJobs, FArchive& InTransferFile, EWriteTasksFlags Flags)
 {
 	int32 InputVersion = ShaderCompileWorkerInputVersion;
 	InTransferFile << InputVersion;
@@ -417,7 +412,7 @@ static bool DoWriteTasksInner(const TArray<FShaderCommonCompileJobPtr>& QueuedJo
 		TransferFile << FormatNone;
 	}
 
-	static TMap<FString, uint32> FormatVersionMap = GetFormatVersionMap();
+	static TMap<FName, uint32> FormatVersionMap = GetFormatVersionMap();
 
 	TransferFile << FormatVersionMap;
 
@@ -529,6 +524,8 @@ static bool DoWriteTasksInner(const TArray<FShaderCommonCompileJobPtr>& QueuedJo
 		}
 	}
 
+	bool bSkipSource = EnumHasAnyFlags(Flags, FShaderCompileWorkerUtil::EWriteTasksFlags::SkipSource);
+	
 	// Write individual shader jobs
 	{
 		int32 SingleJobHeader = ShaderCompileWorkerSingleJobHeader;
@@ -540,7 +537,14 @@ static bool DoWriteTasksInner(const TArray<FShaderCommonCompileJobPtr>& QueuedJo
 		// Serialize all the batched jobs
 		for (int32 JobIndex = 0; JobIndex < QueuedSingleJobs.Num(); JobIndex++)
 		{
-			QueuedSingleJobs[JobIndex]->SerializeWorkerInput(TransferFile);
+			if (bSkipSource)
+			{
+				QueuedSingleJobs[JobIndex]->SerializeWorkerInputNoSource(TransferFile);
+			}
+			else
+			{
+				QueuedSingleJobs[JobIndex]->SerializeWorkerInput(TransferFile);
+			}
 			QueuedSingleJobs[JobIndex]->Input.SerializeSharedInputs(TransferFile, SharedEnvironments, AllShaderParameterStructures);
 		}
 	}
@@ -561,7 +565,14 @@ static bool DoWriteTasksInner(const TArray<FShaderCommonCompileJobPtr>& QueuedJo
 			TransferFile << NumStageJobs;
 			for (int32 Index = 0; Index < NumStageJobs; Index++)
 			{
-				PipelineJob->StageJobs[Index]->SerializeWorkerInput(TransferFile);
+				if (bSkipSource)
+				{
+					PipelineJob->StageJobs[Index]->SerializeWorkerInputNoSource(TransferFile);
+				}
+				else
+				{
+					PipelineJob->StageJobs[Index]->SerializeWorkerInput(TransferFile);
+				}
 				PipelineJob->StageJobs[Index]->Input.SerializeSharedInputs(TransferFile, SharedEnvironments, AllShaderParameterStructures);
 			}
 		}
@@ -597,63 +608,77 @@ static bool DoWriteTasksInner(const TArray<FShaderCommonCompileJobPtr>& QueuedJo
 	return InTransferFile.Close();
 }
 
-const TCHAR* DebugWorkerInputFileName = TEXT("DebugSCW.in");
-const TCHAR* DebugWorkerOutputFileName = TEXT("DebugSCW.out");
+const TCHAR* DebugWorkerInputFileName = TEXT("DebugCompile.in");
+const TCHAR* DebugWorkerOutputFileName = TEXT("DebugCompile.out");
+const TCHAR* DebugCompileArgsFileName = TEXT("DebugCompileArgs.txt");
 
-static FString CreateShaderCompilerWorkerDebugCommandLine(FString DebugWorkerInputFilePath)
+static void WriteShaderCompileWorkerDebugCommandLine(FShaderCommonCompileJob& Job, const FString& JobDebugInfoPath, const FString& InputFilePath, FShaderDebugDataContext& Ctx)
 {
-	// 0 is parent PID, pass zero TTL and KeepInput to make SCW process the single job then exit without deleting the input file
-	return FString::Printf(TEXT("\"%s\" 0 \"DebugSCW\" %s %s -TimeToLive=0.0f -KeepInput"),
-		*DebugWorkerInputFilePath, // working directory for SCW
-		DebugWorkerInputFileName,
-		DebugWorkerOutputFileName);
+
+	TStringBuilder<512> JobArgsPath;
+	FPathViews::Append(JobArgsPath, JobDebugInfoPath, TEXT("DebugCompileArgs.txt"));
+
+	TStringBuilder<512> CmdLine;
+	CmdLine << TEXT("\"") << JobDebugInfoPath << TEXT("\"");
+	CmdLine << TEXT(" 0 \"DebugCompile\" "); // parent PID (not meaningful for debug compile mode) followed by window title
+
+	// output path to the single generated input file for the root job. this will be written in the first stage folder for pipeline jobs,
+	// so make the path relative to the working directory for the current stage. 
+	
+	// note that we pass the path of the compile args txt file to all invocations of FPaths::MakePathRelativeTo in this function 
+	// because it doesn't properly handle normalized paths when the path points to a directory (lack of a trailing / causes an internal 
+	// call to FPaths::GetPath to strip the last folder)
+	FString InputFilePathRelative = InputFilePath;
+	FPaths::MakePathRelativeTo(InputFilePathRelative, JobArgsPath.ToString());
+	CmdLine << InputFilePathRelative << TEXT(" ") << DebugWorkerOutputFileName;
+
+	CmdLine << " -DebugSourceFiles=";
+	TArray<FString> RelativeSourcePaths;
+	RelativeSourcePaths.Reserve(Ctx.DebugSourceFiles.Num());
+	for (FString SourceFile : Ctx.DebugSourceFiles) // note: intentional copy of path here since MakePathRelativeTo modifies in-place
+	{
+		// as above this may refer to multiple source files for different stages of a pipeline job
+		// so make all paths relative to the working directory for this specific job.
+		FPaths::MakePathRelativeTo(SourceFile, JobArgsPath.ToString());
+		RelativeSourcePaths.Add(SourceFile);
+	}
+	CmdLine.JoinQuoted(RelativeSourcePaths, TEXT(","), TEXT("\""));
+	CmdLine << " -TimeToLive=0.0f -KeepInput"; // pass zero TTL and KeepInput to make SCW process the job and exit without deleting the input
+
+	FFileHelper::SaveStringToFile(CmdLine.ToString(), JobArgsPath.ToString());
 }
 
-void FShaderCompileWorkerUtil::DumpDebugCompileInput(TConstArrayView<FShaderCommonCompileJobPtr> QueuedJobs)
+void FShaderCompileWorkerUtil::DumpDebugCompileInput(FShaderCommonCompileJob& Job, FShaderDebugDataContext& Ctx)
 {
-	if (CVarDebugDumpWorkerInputs.GetValueOnAnyThread())
-	{
-		for (const FShaderCommonCompileJobPtr& CommonJob : QueuedJobs)
+	FString CreatedInput;
+	Job.ForEachSingleShaderJob([&CreatedInput, &Job, &Ctx](FShaderCompileJob& SingleJob)
 		{
-			FShaderPipelineCompileJob* PipelineJob = CommonJob->GetShaderPipelineJob();
-			FString DebugWorkerInputFilePath;
-			if (PipelineJob)
-			{
-				// for pipeline jobs, write out the worker input for the whole pipeline, but only for the first stage
-				// would be better to put in a parent folder probably...
-				DebugWorkerInputFilePath = PipelineJob->StageJobs[0]->Input.DumpDebugInfoPath;
-			}
-			else
-			{
-				DebugWorkerInputFilePath = CommonJob->GetSingleShaderJob()->Input.DumpDebugInfoPath;
-			}
-			if (!DebugWorkerInputFilePath.IsEmpty())
+			const FString& JobDebugPath = SingleJob.Input.DumpDebugInfoPath;
+			FString JobInput = JobDebugPath / DebugWorkerInputFileName;
+			if (CreatedInput.IsEmpty())
 			{
 				TArray<FShaderCommonCompileJobPtr> SingleJobArray;
-				SingleJobArray.Add(CommonJob);
-
-				FArchive* DebugWorkerInputFileWriter = IFileManager::Get().CreateFileWriter(*(DebugWorkerInputFilePath / DebugWorkerInputFileName), FILEWRITE_NoFail);
-				DoWriteTasksInner(
+				// export the .in file for just the "root" job; this is either a single job in which case this lambda will only be called once
+				// (and Job == SingleJob), or it's a pipeline job and we want to export a single input file for all jobs and reference it for each stage directory
+				SingleJobArray.Add(&Job);
+				CreatedInput = JobInput;
+				FArchive* DebugWorkerInputFileWriter = IFileManager::Get().CreateFileWriter(*CreatedInput, FILEWRITE_NoFail);
+				WriteTasks(
 					SingleJobArray,
 					*DebugWorkerInputFileWriter,
-					EWriteTasksFlags::CompressTaskFile);	// Always compress the debug files; they are rather large so this saves some disk space
+					// Always compress the debug input files; they are rather large so this saves some disk space
+					EWriteTasksFlags::CompressTaskFile |
+					// Do not include source code in the debug files; this will be read from the debug usf to maintain readability and save disk space
+					EWriteTasksFlags::SkipSource);
 				DebugWorkerInputFileWriter->Close();
 				delete DebugWorkerInputFileWriter;
-
-				FFileHelper::SaveStringToFile(
-					CreateShaderCompilerWorkerDebugCommandLine(DebugWorkerInputFilePath),
-					*(DebugWorkerInputFilePath / TEXT("DebugCompileArgs.txt")));
 			}
-		}
-	}
-}
-
-// Serialize Queued Job information
-bool FShaderCompileWorkerUtil::WriteTasks(const TArray<FShaderCommonCompileJobPtr>& QueuedJobs, FArchive& InTransferFile, EWriteTasksFlags Flags)
-{
-	DumpDebugCompileInput(QueuedJobs);
 
-	return DoWriteTasksInner(QueuedJobs, InTransferFile, Flags);
+			// Always write out the DebugCompileArgs.txt for every stage; this will always run the full pipeline compile for pipeline jobs,
+			// but is just a workflow improvement (so you can navigate to the debug folder for any particular problematic stage and run the full job
+			// without having to know which stage folder contains the input file).
+			WriteShaderCompileWorkerDebugCommandLine(SingleJob, JobDebugPath, CreatedInput, Ctx);
+		});
 }
 
 static void ReadSingleJob(FShaderCompileJob* CurrentJob, FArchive& WorkerOutputFileReader)
diff --git a/Engine/Source/Runtime/RenderCore/Private/ShaderCore.cpp b/Engine/Source/Runtime/RenderCore/Private/ShaderCore.cpp
--- a/Engine/Source/Runtime/RenderCore/Private/ShaderCore.cpp
+++ b/Engine/Source/Runtime/RenderCore/Private/ShaderCore.cpp
@@ -1733,6 +1733,9 @@ public:
 		Job.PreprocessOutput.bSucceeded = Backend->PreprocessShader(Job.Input, Environment, Job.PreprocessOutput);
 		if (Job.PreprocessOutput.bSucceeded && Backend->RequiresSecondaryCompile(Job.Input, Environment, Job.PreprocessOutput))
 		{
+			// check that we're not compiling for a pipeline; this would be theretically possible to support but we'd have to enforce that all jobs in the pipeline
+			// have a secondary compile if any do, to ensure interpolators match across pipeline stages. currently there's no such guarantee so this is a safety measure.
+			checkf(!Job.Input.bCompilingForShaderPipeline, TEXT("Cannot combine shader pipeline interpolator removal optimization with secondary compilation at this time"));
 			Job.SecondaryPreprocessOutput = MakeUnique<FShaderPreprocessOutput>();
 			Job.SecondaryPreprocessOutput->bIsSecondary = true;
 			Job.PreprocessOutput.bSucceeded &= Backend->PreprocessShader(Job.Input, Environment, *Job.SecondaryPreprocessOutput);
@@ -1848,41 +1851,37 @@ public:
 	static void CompileShaderInternal(const IShaderFormat* Compiler, FShaderCompileJob& Job, const FString& WorkingDirectory, FString& OutExceptionMsg, FString& OutExceptionCallstack, int32* CompileCount)
 	{
 		double TimeStart = FPlatformTime::Seconds();
-		if (!Job.PreprocessOutput.GetSucceeded())
-		{
-			PreprocessShaderInternal(Compiler, Job);
-		}
+
+		// the only intended case where we can get to the compile step on a job which failed preprocessing is 
+		// when running the dumped debug compile for such a job. this is intended to allow inspection of the
+		// contents of the FShaderCompileInput (and by extension FShaderCompilerEnvironment(s)) for the job;
+		// so having a check here is a convenience to break the debugger at a point where these things can
+		// easily be inspected.
+		check(Job.PreprocessOutput.GetSucceeded());
 
 		Job.Output.Errors.Append(Job.PreprocessOutput.Errors);
 		
 		// decompress if necessary; this is a no-op if source is not compressed.
 		Job.PreprocessOutput.DecompressCode();
 
-		if (Job.PreprocessOutput.bSucceeded)
+		if (Job.SecondaryPreprocessOutput.IsValid())
 		{
-			if (Job.SecondaryPreprocessOutput.IsValid())
-			{
-				Job.SecondaryOutput = MakeUnique<FShaderCompilerOutput>();
-			}
-			InvokeCompile(Compiler, Job, WorkingDirectory, OutExceptionMsg, OutExceptionCallstack);
-			if (Job.SecondaryOutput.IsValid())
+			Job.SecondaryOutput = MakeUnique<FShaderCompilerOutput>();
+		}
+		InvokeCompile(Compiler, Job, WorkingDirectory, OutExceptionMsg, OutExceptionCallstack);
+		if (Job.SecondaryOutput.IsValid())
+		{
+			Job.Output.bSucceeded = Job.Output.bSucceeded && Job.SecondaryOutput->bSucceeded;
+			// ensure the target field is set on the job output struct as we use it for validation during serialization
+			Job.SecondaryOutput->Target = Job.Input.Target;
+			// also set the job input hash on the output struct for validation purposes
+			Job.SecondaryOutput->ValidateInputHash = Job.Input.Hash;
+			Job.SecondaryOutput->ShaderCode.FinalizeShaderCode();
+			if (Job.SecondaryOutput->bSucceeded)
 			{
-				Job.Output.bSucceeded = Job.Output.bSucceeded && Job.SecondaryOutput->bSucceeded;
-				// ensure the target field is set on the job output struct as we use it for validation during serialization
-				Job.SecondaryOutput->Target = Job.Input.Target;
-				// also set the job input hash on the output struct for validation purposes
-				Job.SecondaryOutput->ValidateInputHash = Job.Input.Hash;
-				Job.SecondaryOutput->ShaderCode.FinalizeShaderCode();
-				if (Job.SecondaryOutput->bSucceeded)
-				{
-					Job.SecondaryOutput->GenerateOutputHash();
-				}
-				CombineOutputs(Compiler, Job);
+				Job.SecondaryOutput->GenerateOutputHash();
 			}
-		}
-		else
-		{
-			Job.Output.bSucceeded = false;
+			CombineOutputs(Compiler, Job);
 		}
 
 		// ensure the target field is set on the job output struct as we use it for validation during serialization
@@ -2365,12 +2364,16 @@ inline int NewlineCharCount(TCHAR First, TCHAR Second)
 // Given an FString containing the contents of a shader source file, populates the given array with contents of
 // that source file with all comments stripped. This is needed since the STB preprocessor itself does not strip 
 // comments.
-void ShaderConvertAndStripComments(const FString& ShaderSource, TArray<ANSICHAR>& OutStripped)
+void ShaderConvertAndStripComments(const FString& ShaderSource, TArray<ANSICHAR>& OutStripped, EConvertAndStripFlags Flags)
 {
+	bool bSkipSimdPadding = EnumHasAnyFlags(Flags, EConvertAndStripFlags::NoSimdPadding);
+	// need extra null chars at the end of the stripped string - one null terminator plus 15 optional padding characters for SIMD-safe reads
+	const uint8 NullCharCount = bSkipSimdPadding ? 1 : 16;
+	
 	// STB preprocessor does not strip comments, so we do so here before returning the loaded source
 	// Doing so is barely more costly than the memcopy we require anyways so has negligible overhead.
 	// Reserve worst case (i.e. assuming there are no comments at all) to avoid reallocation
-	int32 BufferSize = ShaderSource.Len() + 16;		// need extra for null terminator plus padding for SSE read operations at the end of the buffer
+	int32 BufferSize = ShaderSource.Len() + NullCharCount;
 	OutStripped.SetNumUninitialized(BufferSize);
 
 	ANSICHAR* CurrentOut = OutStripped.GetData();
@@ -2597,9 +2600,10 @@ void ShaderConvertAndStripComments(const FString& ShaderSource, TArray<ANSICHAR>
 			}
 		}
 	}
+
 	// Null terminate after comment-stripped copy, plus 15 zero padding characters for SSE safe reads
-	check(CurrentOut + 16 <= (OutStripped.GetData() + BufferSize));
-	for (int32 TerminateAndPadIndex = 0; TerminateAndPadIndex < 16; TerminateAndPadIndex++)
+	check(CurrentOut + NullCharCount <= (OutStripped.GetData() + BufferSize));
+	for (int32 TerminateAndPadIndex = 0; TerminateAndPadIndex < NullCharCount; TerminateAndPadIndex++)
 	{
 		*CurrentOut++ = 0;
 	}
@@ -4306,7 +4310,7 @@ void FShaderCompileJob::SerializeOutput(FShaderCacheSerializeContext& Ctx, int32
 	}
 }
 
-void FShaderCompileJob::OnComplete()
+void FShaderCompileJob::OnComplete(FShaderDebugDataContext& Ctx)
 {
 	if (Input.NeedsOriginalShaderSource())
 	{
@@ -4342,25 +4346,35 @@ void FShaderCompileJob::OnComplete()
 	// job which executed compilation)
 	if (Input.DumpDebugInfoEnabled())
 	{
+		bool bShouldDump = CVarDumpDebugInfoForCacheHits.GetValueOnAnyThread() || !JobStatusPtr->WasCompilationSkipped();
+
 		check(ShaderFormat);
 		if (!PreprocessOutput.bSucceeded)
 		{
 			// output minimal debug info for preprocessing failures
-			ShaderFormat->OutputDebugDataMinimal(Input);
+			ShaderFormat->OutputDebugDataMinimal(Input, Ctx);
 		}
-		else if (CVarDumpDebugInfoForCacheHits.GetValueOnAnyThread() || Output.CompileTime > 0.0f)
+		else if (bShouldDump)
 			// if we only want debug info for jobs which actually compiled, check the CompileTime
 			// (jobs deserialized from the cache/wait list/ddc will have a compiletime of 0.0)
 		{
 			if (SecondaryPreprocessOutput.IsValid() && SecondaryOutput.IsValid())
 			{
-				ShaderFormat->OutputDebugData(Input, PreprocessOutput, *SecondaryPreprocessOutput, Output, *SecondaryOutput);
+				ShaderFormat->OutputDebugData(Input, PreprocessOutput, *SecondaryPreprocessOutput, Output, *SecondaryOutput, Ctx);
 			}
 			else
 			{
-				ShaderFormat->OutputDebugData(Input, PreprocessOutput, Output);
+				ShaderFormat->OutputDebugData(Input, PreprocessOutput, Output, Ctx);
 			}
 		}
+
+		if (bShouldDump && !Ctx.bIsPipeline)
+		{
+			// only need to output the debug compile input/cmdline here if we're not part of a pipeline, otherwise 
+			// it's done at the pipeline level and symlinked in the appropriate place
+			// note that we can't use FShaderCompilerInput::bCompilingForShaderPipeline as that is not set for all stages of a pipeline job
+			FShaderCompileWorkerUtil::DumpDebugCompileInput(*this, Ctx);
+		}
 	}
 
 	if (bHasSymbols)
@@ -4469,6 +4483,21 @@ void FShaderCompileJob::SerializeWorkerInput(FArchive& Ar)
 	}
 }
 
+void FShaderCompileJob::SerializeWorkerInputNoSource(FArchive& Ar)
+{
+	FShaderSource SourceTemp = MoveTemp(PreprocessOutput.PreprocessedSource);
+	bool bHasSecondary = SecondaryPreprocessOutput.IsValid();
+	FShaderSource SecondarySourceTemp = bHasSecondary ? MoveTemp(SecondaryPreprocessOutput->PreprocessedSource) : FShaderSource();
+
+	SerializeWorkerInput(Ar);
+
+	PreprocessOutput.PreprocessedSource = MoveTemp(SourceTemp);
+	if (bHasSecondary)
+	{
+		SecondaryPreprocessOutput->PreprocessedSource = MoveTemp(SecondarySourceTemp);
+	}
+}
+
 void FShaderCompileJob::SerializeArtifact(FArchive& Ar)
 {
 	SerializeWorkerInput(Ar);
@@ -4560,11 +4589,19 @@ void FShaderPipelineCompileJob::SerializeOutput(FShaderCacheSerializeContext& Ct
 	}
 }
 
-void FShaderPipelineCompileJob::OnComplete()
+void FShaderPipelineCompileJob::OnComplete(FShaderDebugDataContext& Ctx)
 {
+	Ctx.bIsPipeline = true;
 	for (int32 Index = 0, Num = StageJobs.Num(); Index < Num; ++Index)
 	{
-		StageJobs[Index]->OnComplete();
+		StageJobs[Index]->OnComplete(Ctx);
+	}
+
+	bool bShouldDump = !Ctx.DebugSourceFiles.IsEmpty() && (CVarDumpDebugInfoForCacheHits.GetValueOnAnyThread() || !JobStatusPtr->WasCompilationSkipped());
+	if (bShouldDump)
+	{
+		check(Ctx.DebugSourceFiles.Num() == StageJobs.Num());
+		FShaderCompileWorkerUtil::DumpDebugCompileInput(*this, Ctx);
 	}
 }
 
diff --git a/Engine/Source/Runtime/RenderCore/Public/ShaderCompilerCore.h b/Engine/Source/Runtime/RenderCore/Public/ShaderCompilerCore.h
--- a/Engine/Source/Runtime/RenderCore/Public/ShaderCompilerCore.h
+++ b/Engine/Source/Runtime/RenderCore/Public/ShaderCompilerCore.h
@@ -28,7 +28,7 @@ class FShaderPipelineCompileJob;
 typedef TSharedPtr<TArray<ANSICHAR>, ESPMode::ThreadSafe> FShaderSharedAnsiStringPtr;
 
 // this is for the protocol, not the data, bump if FShaderCompilerInput/FShaderPreprocessOutput serialization, SerializeWorkerInput or ProcessInputFromArchive changes.
-inline const int32 ShaderCompileWorkerInputVersion = 29;
+inline const int32 ShaderCompileWorkerInputVersion = 30;
 // this is for the protocol, not the data, bump if FShaderCompilerOutput or WriteToOutputArchive changes.
 inline const int32 ShaderCompileWorkerOutputVersion = 28;
 // this is for the protocol, not the data.
@@ -172,7 +172,7 @@ namespace FOodleDataCompression
 enum class EShaderDebugInfoFlags : uint8
 {
 	Default = 0,
-	DirectCompileCommandLine = 1 << 0,
+	DirectCompileCommandLine UE_DEPRECATED(5.7, "DirectCompile has been renamed DebugCompile and DebugCompileArgs.txt is dumped by default") = 1 << 0,
 	InputHash = 1 << 1,
 	Diagnostics = 1 << 2,
 	ShaderCodeBinary = 1 << 3,
@@ -184,6 +184,12 @@ ENUM_CLASS_FLAGS(EShaderDebugInfoFlags)
 
 using FShaderCompilerInputHash = FBlake3Hash;
 
+struct FShaderDebugDataContext
+{
+	bool bIsPipeline = false;
+	TArray<FString> DebugSourceFiles;
+};
+
 /** Struct that gathers all readonly inputs needed for the compilation of a single shader. */
 struct FShaderCompilerInput
 {
@@ -646,10 +652,16 @@ extern RENDERCORE_API bool CheckVirtualShaderFilePath(FStringView VirtualPath, T
  */
 extern RENDERCORE_API void FixupShaderFilePath(FString& VirtualFilePath, EShaderPlatform ShaderPlatform, const FName* ShaderPlatformName);
 
+enum class EConvertAndStripFlags : uint8
+{
+	None = 0,
+	NoSimdPadding = 1 << 0,
+};
+ENUM_CLASS_FLAGS(EConvertAndStripFlags);
 /**
  * Utility function to strip comments and convert source to ANSI, useful for preprocessing
  */
-extern RENDERCORE_API void ShaderConvertAndStripComments(const FString& ShaderSource, TArray<ANSICHAR>& OutStripped);
+extern RENDERCORE_API void ShaderConvertAndStripComments(const FString& ShaderSource, TArray<ANSICHAR>& OutStripped, EConvertAndStripFlags Flags = EConvertAndStripFlags::None);
 
 /**
  * Loads the shader file with the given name.
diff --git a/Engine/Source/Runtime/RenderCore/Public/ShaderCompilerJobTypes.h b/Engine/Source/Runtime/RenderCore/Public/ShaderCompilerJobTypes.h
--- a/Engine/Source/Runtime/RenderCore/Public/ShaderCompilerJobTypes.h
+++ b/Engine/Source/Runtime/RenderCore/Public/ShaderCompilerJobTypes.h
@@ -101,12 +101,22 @@ public:
 		IsDuplicate.store(false);
 	}
 
+	bool WasCompilationSkipped() const
+	{
+		EShaderCompileJobStatus LocalStatus = GetStatus();
+
+		return GetIsDuplicate() ||
+		(LocalStatus == EShaderCompileJobStatus::Skipped) ||
+		(LocalStatus == EShaderCompileJobStatus::CompleteFoundInCache) ||
+		(LocalStatus == EShaderCompileJobStatus::CompleteFoundInDDC);
+	}
+
 	void SetStatus(EShaderCompileJobStatus NewStatus)
 	{
 		Status.store(NewStatus);
 	}
 
-	EShaderCompileJobStatus GetStatus()
+	EShaderCompileJobStatus GetStatus() const
 	{
 		return Status.load();
 	}
@@ -126,7 +136,7 @@ public:
 		InputHash = InInputHash;
 	}
 
-	const FShaderCompilerInputHash& GetInputHash()
+	const FShaderCompilerInputHash& GetInputHash() const
 	{
 		return InputHash ? *InputHash : FShaderCompilerInputHash::Zero;
 	}
@@ -240,9 +250,12 @@ public:
 	FShaderPipelineCompileJob* GetShaderPipelineJob();
 	const FShaderPipelineCompileJob* GetShaderPipelineJob() const;
 
-	// Executed for all jobs (including those read from cache) on completion.
+	UE_DEPRECATED(5.7, "Use overload of OnComplete accepting a context object")
 	virtual void OnComplete() = 0;
 
+	// Executed for all jobs (including those read from cache) on completion.
+	virtual void OnComplete(FShaderDebugDataContext& Ctx) = 0;
+
 	virtual void AppendDebugName(FStringBuilderBase& OutName) const = 0;
 
 	bool Equals(const FShaderCommonCompileJob& Rhs) const;
@@ -345,7 +358,10 @@ public:
 		SerializeOutput(Ctx, 0);
 	}
 
-	virtual RENDERCORE_API void OnComplete() override;
+	UE_DEPRECATED(5.7, "OnComplete now accepts a FDebugDataContext object")
+	virtual void OnComplete() override {}
+
+	virtual RENDERCORE_API void OnComplete(FShaderDebugDataContext& Ctx) override;
 	
 	virtual RENDERCORE_API void AppendDebugName(FStringBuilderBase& OutName) const override;
 
@@ -355,6 +371,8 @@ public:
 	// Serializes only the subset of data written by ShaderCompiler and read from SCW when using worker processes.
 	RENDERCORE_API void SerializeWorkerInput(FArchive& Ar);
 
+	void SerializeWorkerInputNoSource(FArchive& Ar);
+
 	// Serializes the compile job for a cook artifact for later analysis in a commandlet.
 	RENDERCORE_API void SerializeArtifact(FArchive& Ar);
 
@@ -409,7 +427,9 @@ public:
 
 	virtual RENDERCORE_API void SerializeOutput(FShaderCacheSerializeContext& Ctx) override;
 
-	virtual RENDERCORE_API void OnComplete() override;
+	UE_DEPRECATED(5.7, "OnComplete now accepts a FDebugDataContext object")
+	void OnComplete() override {};
+	virtual RENDERCORE_API void OnComplete(FShaderDebugDataContext& Ctx) override;
 	virtual RENDERCORE_API void AppendDebugName(FStringBuilderBase& OutName) const override;
 	virtual RENDERCORE_API void AppendDiagnostics(FString& OutDiagnostics, int32 InJobIndex, int32 InNumJobs, const TCHAR* Indentation = nullptr) const override final;
 
