[shaders] fix debug compile of jobs with secondary outputs

diff --git a/Engine/Source/Programs/ShaderCompileWorker/Private/ShaderCompileWorker.cpp b/Engine/Source/Programs/ShaderCompileWorker/Private/ShaderCompileWorker.cpp
--- a/Engine/Source/Programs/ShaderCompileWorker/Private/ShaderCompileWorker.cpp
+++ b/Engine/Source/Programs/ShaderCompileWorker/Private/ShaderCompileWorker.cpp
@@ -703,7 +703,7 @@ private:
 
 		GNumProcessedJobs = 0;
 
-		auto SetupDebugCompile = [this, &ReceivedFormatVersionMap](FShaderCompileJob& Job, const TCHAR* SourcePath)
+		auto SetupDebugCompile = [this, &ReceivedFormatVersionMap](FShaderCompileJob& Job, const TCHAR* SourcePath, const TCHAR* SecondarySourcePath = nullptr)
 			{
 				Job.Input.DebugInfoFlags &= EShaderDebugInfoFlags::CompileFromDebugUSF;
 
@@ -725,6 +725,21 @@ private:
 					ShaderConvertAndStripComments(DebugSource, Stripped.GetCharArray(), EConvertAndStripFlags::NoSimdPadding);
 					Job.PreprocessOutput.EditSource() = MoveTemp(Stripped);
 				}
+
+				if (SecondarySourcePath != nullptr)
+				{
+					FString SecondaryDebugSource;
+					bSuccess = FFileHelper::LoadFileToString(SecondaryDebugSource, SecondarySourcePath);
+					if (bSuccess)
+					{
+						// strip comments from source when loading from a debug source file. some backends don't handle the comments that the debug dump inserts properly.
+						FAnsiString SecondaryStripped;
+						// Since we are directly assigning the stripped result to an FShaderSource object which automatically appends SIMD patting as part of the Set function,
+						// we don't need to also add padding during the convert-and-strip operation.
+						ShaderConvertAndStripComments(SecondaryDebugSource, SecondaryStripped.GetCharArray(), EConvertAndStripFlags::NoSimdPadding);
+						Job.SecondaryPreprocessOutput->EditSource() = MoveTemp(SecondaryStripped);
+					}
+				}
 			};
 
 		// Individual jobs
@@ -751,10 +766,20 @@ private:
 				Job.Input.DeserializeSharedInputs(InputFile, SharedEnvironments, ParameterStructures);
 				if (bIsDebugCompile)
 				{
-
+					bool bHasSecondary = Job.SecondaryPreprocessOutput.IsValid();
 					// this should only be used in debug compile mode; in this case we expect either a single job or a single pipeline job
-					check(DebugSourceFiles.Num() == 1 && NumSingleJobs == 1);
-					SetupDebugCompile(Job, *DebugSourceFiles[0]);
+					if (bHasSecondary)
+					{
+						// jobs with secondary compiles should have two source files
+						check(DebugSourceFiles.Num() == 2 && NumSingleJobs == 1);
+						SetupDebugCompile(Job, *DebugSourceFiles[0], *DebugSourceFiles[1]);
+					}
+					else
+					{
+						// normal jobs should just have a single source file
+						check(DebugSourceFiles.Num() == 1 && NumSingleJobs == 1);
+						SetupDebugCompile(Job, *DebugSourceFiles[0]);
+					}
 				}
 
 				// Process the job.
