[SCW] Add command line argument to override debug output path in ShaderCompileWorker to debugging.

When diagnosing shader job artifacts we might want to dump new debug information but to a different location as paths from the build machine might mismatch with the local machine.
This change allows to override the debug info path via a new command line argument `-DumpDebugInfoPath=PATH` in the SCW.

#rnx

diff --git a/Engine/Source/Programs/ShaderCompileWorker/Private/ShaderCompileWorker.cpp b/Engine/Source/Programs/ShaderCompileWorker/Private/ShaderCompileWorker.cpp
--- a/Engine/Source/Programs/ShaderCompileWorker/Private/ShaderCompileWorker.cpp
+++ b/Engine/Source/Programs/ShaderCompileWorker/Private/ShaderCompileWorker.cpp
@@ -231,6 +231,7 @@ public:
 	bool DisableFileWrite = false;
 	bool KeepInput = false;
 	TArray<FString> DebugSourceFiles;
+	FString DumpDebugInfoPathOverride;
 
 	FWorkLoop(const TCHAR* ParentProcessIdText,const TCHAR* InWorkingDirectory,const TCHAR* InInputFilename,const TCHAR* InOutputFilename, TMap<FName, uint32>& InFormatVersionMap)
 	:	ParentProcessId(FCString::Atoi(ParentProcessIdText))
@@ -271,6 +272,10 @@ public:
 					DebugSourceFile = InWorkingDirectory / DebugSourceFile;
 				}
 			}
+			else if (Switch.StartsWith(TEXT("DumpDebugInfoPath=")))
+			{
+				DumpDebugInfoPathOverride = FString(Switch.RightChop(18));
+			}
 		}
 	}
 
@@ -716,6 +721,11 @@ private:
 				FString DebugSource;
 				bool bSuccess = FFileHelper::LoadFileToString(DebugSource, SourcePath);
 
+				if (!DumpDebugInfoPathOverride.IsEmpty())
+				{
+					Job.Input.DumpDebugInfoPath = DumpDebugInfoPathOverride;
+				}
+
 				if (bSuccess)
 				{
 					// strip comments from source when loading from a debug source file. some backends don't handle the comments that the debug dump inserts properly.
