[Shaders] Fix adding debug info flag when running DebugCompile.in file through SCW.

#test Ran SCW with DebugCompile.in file locally. This fixed the missing error messages.
#rnx

diff --git a/Engine/Source/Programs/ShaderCompileWorker/Private/ShaderCompileWorker.cpp b/Engine/Source/Programs/ShaderCompileWorker/Private/ShaderCompileWorker.cpp
--- a/Engine/Source/Programs/ShaderCompileWorker/Private/ShaderCompileWorker.cpp
+++ b/Engine/Source/Programs/ShaderCompileWorker/Private/ShaderCompileWorker.cpp
@@ -710,14 +710,13 @@ private:
 
 		auto SetupDebugCompile = [this, &ReceivedFormatVersionMap](FShaderCompileJob& Job, const TCHAR* SourcePath, const TCHAR* SecondarySourcePath = nullptr)
 			{
-				Job.Input.DebugInfoFlags &= EShaderDebugInfoFlags::CompileFromDebugUSF;
+				Job.Input.DebugInfoFlags |= EShaderDebugInfoFlags::CompileFromDebugUSF;
 
 				checkf(ReceivedFormatVersionMap[Job.Input.ShaderFormat] == FormatVersionMap[Job.Input.ShaderFormat],
 						TEXT("Format mismatch for shader format %s; ensure SDK/tools for this platform match the build where this debug compile input was generated.\n")
 						TEXT("This check is skippable, with the caveat that this debug compilation may fail for unexpected reasons and/or results may be different."),
 						*Job.Input.ShaderFormat.ToString());
 
-				Job.Input.DebugInfoFlags &= EShaderDebugInfoFlags::CompileFromDebugUSF;
 				FString DebugSource;
 				bool bSuccess = FFileHelper::LoadFileToString(DebugSource, SourcePath);
 
