[ShaderCompileWorker]
* Added so uba path for reusing process supports TimeToLive.

#Synced-CL 49713413

diff --git a/Engine/Plugins/UbaController/Source/Private/UbaJobProcessor.cpp b/Engine/Plugins/UbaController/Source/Private/UbaJobProcessor.cpp
--- a/Engine/Plugins/UbaController/Source/Private/UbaJobProcessor.cpp
+++ b/Engine/Plugins/UbaController/Source/Private/UbaJobProcessor.cpp
@@ -42,6 +42,12 @@ namespace UbaJobProcessorOptions
 		HeartBeatInterval,
 		TEXT("Time between heart beat log messages"));
 
+	static float ProcessIdleTime = 15.0f;
+	static FAutoConsoleVariableRef CVarProcessIdleTime(
+		TEXT("r.UbaController.ProcessIdleTime"),
+		ProcessIdleTime,
+		TEXT("ShaderCompileWorker idle time until it exits. Only used when process reuse is enabled"));
+
 	static bool bAutoLaunchVisualizer = false;
 	static FAutoConsoleVariableRef CVarAutoLaunchVisualizer(
 		TEXT("r.UbaController.AutoLaunchVisualizer"),
@@ -277,9 +283,11 @@ void FUbaJobProcessor::RunTaskWithUba(FDistributedBuildTask* Task)
 		SessionServer_RegisterNewDirectory(UbaSessionServer, TCHAR_TO_UBASTRING(*AdditionalOutputFolder));
 	}
 
+	float TimeToLive = UbaJobProcessorOptions::bAllowProcessReuse ? UbaJobProcessorOptions::ProcessIdleTime : 0.0f;
+
 	FString InputFileName = FPaths::GetCleanFilename(Data.InputFileName);
 	FString OutputFileName = FPaths::GetCleanFilename(Data.OutputFileName);
-	FString Parameters = FString::Printf(TEXT("\"%s/\" %d 0 \"%s\" \"%s\" %s "), *Data.WorkingDirectory, Data.DispatcherPID, *InputFileName, *OutputFileName, *Data.ExtraCommandArgs);
+	FString Parameters = FString::Printf(TEXT("\"%s/\" %d 0 \"%s\" \"%s\" %s -TimeToLive=%.1f"), *Data.WorkingDirectory, Data.DispatcherPID, *InputFileName, *OutputFileName, *Data.ExtraCommandArgs, TimeToLive);
 	FString AppDir = FPaths::GetPath(Data.Command);
 	
 	struct ExitedInfo
diff --git a/Engine/Source/Programs/ShaderCompileWorker/Private/ShaderCompileWorker.cpp b/Engine/Source/Programs/ShaderCompileWorker/Private/ShaderCompileWorker.cpp
--- a/Engine/Source/Programs/ShaderCompileWorker/Private/ShaderCompileWorker.cpp
+++ b/Engine/Source/Programs/ShaderCompileWorker/Private/ShaderCompileWorker.cpp
@@ -353,24 +353,24 @@ public:
 				using ubachar = char;
 #endif
 				
-				using UbaRequestNextProcessFunc = bool(uint32 prevExitCode, ubachar* outArguments, uint32 outArgumentsCapacity);
-				static UbaRequestNextProcessFunc* RequestNextProcess = nullptr;
+				using UbaRequestNextProcessFunc2 = bool(uint32 prevExitCode, ubachar* outArguments, uint32 outArgumentsCapacity, uint32 timeOutMs, bool* outShouldExit);
+				static UbaRequestNextProcessFunc2* RequestNextProcess = nullptr;
 				if (!RequestNextProcess)
 				{
 #if PLATFORM_WINDOWS
 					if (HMODULE UbaDetoursModule = GetUbaModule())
 					{
-						RequestNextProcess = (UbaRequestNextProcessFunc*)(void*)GetProcAddress(UbaDetoursModule, "UbaRequestNextProcess");
+						RequestNextProcess = (UbaRequestNextProcessFunc2*)(void*)GetProcAddress(UbaDetoursModule, "UbaRequestNextProcess2");
 					}
 #elif PLATFORM_MAC
 					if (void* UbaDetoursHandle = dlopen("libUbaDetours.dylib", RTLD_LAZY))
 					{
-						RequestNextProcess = (UbaRequestNextProcessFunc*)(void*)dlsym(UbaDetoursHandle, "UbaRequestNextProcess");
+						RequestNextProcess = (UbaRequestNextProcessFunc2*)(void*)dlsym(UbaDetoursHandle, "UbaRequestNextProcess2");
 					}
 #elif PLATFORM_LINUX
 					if (void* UbaDetoursHandle = dlopen("libUbaDetours.so", RTLD_LAZY))
 					{
-						RequestNextProcess = (UbaRequestNextProcessFunc*)(void*)dlsym(UbaDetoursHandle, "UbaRequestNextProcess");
+						RequestNextProcess = (UbaRequestNextProcessFunc2*)(void*)dlsym(UbaDetoursHandle, "UbaRequestNextProcess2");
 					}
 #endif
 					if (!RequestNextProcess)
@@ -381,7 +381,18 @@ public:
 
 				// Request new process
 				ubachar Temp[1024];
-				if (!RequestNextProcess(0, Temp, 1024))
+				uint32 TimeoutMs = 0;
+				if (TimeToLive != 0)
+				{
+					double TimeSinceLastCompile = FPlatformTime::Seconds() - GLastCompileTime;
+					if (TimeSinceLastCompile < TimeToLive)
+					{
+						TimeoutMs = uint32((TimeToLive - TimeSinceLastCompile)*1000.0);
+					}
+				}
+
+				bool ShouldExit = false; // Ignore this, we always exit after timeout
+				if (!RequestNextProcess(0, Temp, 1024, TimeoutMs, &ShouldExit))
 				{
 					break; // No process available, exit loop
 				}
