Fixed compile errors in various programs when compiling with portable toolchain

#Synced-CL 0

diff --git a/Engine/Source/Programs/ShaderCompileWorker/Private/ShaderCompileWorker.cpp b/Engine/Source/Programs/ShaderCompileWorker/Private/ShaderCompileWorker.cpp
--- a/Engine/Source/Programs/ShaderCompileWorker/Private/ShaderCompileWorker.cpp
+++ b/Engine/Source/Programs/ShaderCompileWorker/Private/ShaderCompileWorker.cpp
@@ -24,6 +24,7 @@
 #include "SocketSubsystem.h"
 
 #if PLATFORM_MAC
+#include <dlfcn.h>
 #include <mach-o/dyld.h>
 #elif PLATFORM_LINUX
 #include <link.h>
diff --git a/Engine/Source/Programs/ShaderCompileWorker/ShaderCompileWorker.Build.cs b/Engine/Source/Programs/ShaderCompileWorker/ShaderCompileWorker.Build.cs
--- a/Engine/Source/Programs/ShaderCompileWorker/ShaderCompileWorker.Build.cs
+++ b/Engine/Source/Programs/ShaderCompileWorker/ShaderCompileWorker.Build.cs
@@ -7,6 +7,7 @@ public class ShaderCompileWorker : ModuleRules
 {
 	public ShaderCompileWorker(ReadOnlyTargetRules Target) : base(Target)
 	{
+		bRequiresPlatformSDK = true;
 		bTreatAsEngineModule = true;
 
 		PrivateDependencyModuleNames.AddRange(
diff --git a/Engine/Source/Programs/SlateUGS/SlateUGS.Build.cs b/Engine/Source/Programs/SlateUGS/SlateUGS.Build.cs
--- a/Engine/Source/Programs/SlateUGS/SlateUGS.Build.cs
+++ b/Engine/Source/Programs/SlateUGS/SlateUGS.Build.cs
@@ -6,6 +6,8 @@ public class SlateUGS : ModuleRules
 {
 	public SlateUGS(ReadOnlyTargetRules Target) : base(Target)
 	{
+		bRequiresPlatformSDK = true;
+
 		PublicIncludePathModuleNames.Add("Launch");
 
 		PrivateDependencyModuleNames.AddRange(
diff --git a/Engine/Source/Programs/SlateUGS/SlateUGS.Target.cs b/Engine/Source/Programs/SlateUGS/SlateUGS.Target.cs
--- a/Engine/Source/Programs/SlateUGS/SlateUGS.Target.cs
+++ b/Engine/Source/Programs/SlateUGS/SlateUGS.Target.cs
@@ -10,6 +10,7 @@ public class SlateUGSTarget : TargetRules
 	{
 		Type = TargetType.Program;
 		LinkType = TargetLinkType.Monolithic;
+		IncludeOrderVersion = EngineIncludeOrderVersion.Latest;
 
 		LaunchModuleName = "SlateUGS";
 
