Shaders: Moved SCW input processing and output writing to a shared location

#Synced-CL 50253651
#Synced-CL 50253651

diff --git a/Engine/Source/Programs/ShaderCompileWorker/Private/ShaderCompileWorker.cpp b/Engine/Source/Programs/ShaderCompileWorker/Private/ShaderCompileWorker.cpp
--- a/Engine/Source/Programs/ShaderCompileWorker/Private/ShaderCompileWorker.cpp
+++ b/Engine/Source/Programs/ShaderCompileWorker/Private/ShaderCompileWorker.cpp
@@ -4,23 +4,9 @@
 // ShaderCompileWorker.cpp : Defines the entry point for the console application.
 //
 
-#if UE_ENABLE_INCLUDE_ORDER_DEPRECATED_IN_5_4
-#include "CoreMinimal.h"
-#endif
-#include "Misc/Compression.h"
 #include "RequiredProgramMainCPPInclude.h"
-#include "ShaderCompilerCore.h"
-#include "ShaderCompilerCommon.h"
-#include "ShaderCompilerJobTypes.h"
-#include "ShaderCore.h"
 #include "HAL/ExceptionHandling.h"
-#include "Interfaces/IShaderFormat.h"
-#include "Interfaces/IShaderFormatModule.h"
-#include "Interfaces/ITargetPlatformManagerModule.h"
-#include "RHIShaderFormatDefinitions.inl"
-#include "ShaderCompilerCommon.h"
 #include "ShaderCompileWorkerUtil.h"
-#include "Serialization/MemoryReader.h"
 #include "SocketSubsystem.h"
 
 #if PLATFORM_MAC
@@ -102,124 +88,16 @@ static inline void ExitWithoutCrash(FSCWErrorCode::ECode ErrorCode, const FStrin
 	UE_LOG(LogShaders, Fatal, TEXT("%s"), *Message);
 }
 
-static const TArray<const IShaderFormat*>& GetShaderFormats()
-{
-	static bool bInitialized = false;
-	static TArray<const IShaderFormat*> Results;
-
-	if (!bInitialized)
-	{
-		bInitialized = true;
-		Results.Empty(Results.Num());
-
-		TArray<FName> Modules;
-		FModuleManager::Get().FindModules(SHADERFORMAT_MODULE_WILDCARD, Modules);
-
-		if (!Modules.Num())
-		{
-			ExitWithoutCrash(FSCWErrorCode::NoTargetShaderFormatsFound, TEXT("No target shader formats found!"));
-		}
-
-		for (int32 Index = 0; Index < Modules.Num(); Index++)
-		{
-			IShaderFormat* Format = FModuleManager::LoadModuleChecked<IShaderFormatModule>(Modules[Index]).GetShaderFormat();
-			if (Format != nullptr)
-			{
-				Results.Add(Format);
-			}
-		}
-	}
-	return Results;
-}
-
 
-static void UpdateFileSize(FArchive& OutputFile, int64 FileSizePosition)
+static const FString& GetLocalHostName()
 {
-	int64 Current = OutputFile.Tell();
-	OutputFile.Seek(FileSizePosition);
-	OutputFile << Current;
-	OutputFile.Seek(Current);
-};
-
-static const TCHAR* GetLocalHostname(int32* OutHostnameLength = nullptr)
-{
-	static FString Hostname;
-	if (Hostname.IsEmpty())
-	{
-		ISocketSubsystem::Get()->GetHostName(Hostname);
-	}
-	if (OutHostnameLength)
+	static FString HostName = []
 	{
-		*OutHostnameLength = Hostname.Len();
-	}
-	return *Hostname;
-}
-
-static int64 WriteOutputFileHeader(FArchive& OutputFile, int32 CallstackLength, const TCHAR* Callstack,
-	int32 ExceptionInfoLength, const TCHAR* ExceptionInfo)
-{
-	TRACE_CPUPROFILER_EVENT_SCOPE(WriteOutputFileHeader);
-
-	int64 FileSizePosition = 0;
-	int32 OutputVersion = ShaderCompileWorkerOutputVersion;
-	OutputFile << OutputVersion;
-
-	int64 FileSize = 0;
-	// Get the position of the Size value to be patched in as the shader progresses
-	FileSizePosition = OutputFile.Tell();
-	OutputFile << FileSize;
-
-	OutputFile << GWorkerDiagnostics;
-
-	OutputFile << GNumProcessedJobs;
-
-	// Note: Can't use FStrings here as SEH can't be used with destructors
-	OutputFile << CallstackLength;
-
-	OutputFile << ExceptionInfoLength;
-
-	int32 HostnameLength = 0;
-	const TCHAR* Hostname = GetLocalHostname(&HostnameLength);
-	OutputFile << HostnameLength;
-
-	if (GWorkerDiagnostics.ErrorCode != FSCWErrorCode::Success)
-	{
-		if (CallstackLength > 0)
-		{
-			OutputFile.Serialize((void*)Callstack, CallstackLength * sizeof(TCHAR));
-		}
-
-		if (ExceptionInfoLength > 0)
-		{
-			OutputFile.Serialize((void*)ExceptionInfo, ExceptionInfoLength * sizeof(TCHAR));
-		}
-
-		if (HostnameLength > 0)
-		{
-			OutputFile.Serialize((void*)Hostname, HostnameLength * sizeof(TCHAR));
-		}
-
-		// Store available and used physical memory of host machine on OOM error
-		if (GWorkerDiagnostics.ErrorCode == FSCWErrorCode::OutOfMemory)
-		{
-			FPlatformMemoryStats MemoryStats = FPlatformMemory::GetStats();
-
-			OutputFile
-				<< MemoryStats.AvailablePhysical
-				<< MemoryStats.AvailableVirtual
-				<< MemoryStats.UsedPhysical
-				<< MemoryStats.PeakUsedPhysical
-				<< MemoryStats.UsedVirtual
-				<< MemoryStats.PeakUsedVirtual
-				;
-		}
-	}
-
-	// Reset error code as it can be receive a new value now
-	FSCWErrorCode::Reset();
-
-	UpdateFileSize(OutputFile, FileSizePosition);
-	return FileSizePosition;
+		FString Output;
+		ISocketSubsystem::Get()->GetHostName(Output);
+		return Output;
+	}();
+	return HostName;
 }
 
 
@@ -234,14 +112,13 @@ public:
 	TArray<FString> DebugSourceFiles;
 	FString DumpDebugInfoPathOverride;
 
-	FWorkLoop(const TCHAR* ParentProcessIdText,const TCHAR* InWorkingDirectory,const TCHAR* InInputFilename,const TCHAR* InOutputFilename, TMap<FName, uint32>& InFormatVersionMap)
+	FWorkLoop(const TCHAR* ParentProcessIdText,const TCHAR* InWorkingDirectory,const TCHAR* InInputFilename,const TCHAR* InOutputFilename)
 	:	ParentProcessId(FCString::Atoi(ParentProcessIdText))
 	,	WorkingDirectory(InWorkingDirectory)
 	,	InputFilename(InInputFilename)
 	,	OutputFilename(InOutputFilename)
 	,	InputFilePath(FString(InWorkingDirectory) / InInputFilename)
 	,	OutputFilePath(FString(InWorkingDirectory) / InOutputFilename)
-	,	FormatVersionMap(InFormatVersionMap)
 	{
 		TArray<FString> Tokens, Switches;
 		FCommandLine::Parse(FCommandLine::Get(), Tokens, Switches);
@@ -310,19 +187,32 @@ public:
 
 				UE_LOG(LogShaders, Log, TEXT("Processing shader"));
 
-				ProcessInputFromArchive(InputFilePtr.Get(), SingleJobs, PipelineJobs, PipelineJobNames);
+				auto LoadVirtualSourceFile = [this](int32 VirtualSourceIndex) -> TArray<uint8>
+				{
+					TArray<uint8> Output;
+					TStringBuilder<512> VirtualSourcePath;
+					VirtualSourcePath << InputFilePath << VirtualSourceIndex;
+					FFileHelper::LoadFileToArray(Output, *VirtualSourcePath, FILEREAD_NoFail);
+					return Output;
+				};
+
+				FShaderCompileWorkerUtil::ProcessInputFromArchive(InputFilePtr.Get(), DumpDebugInfoPathOverride, DebugSourceFiles,
+					LoadVirtualSourceFile, GNumProcessedJobs, SingleJobs, PipelineJobs, PipelineJobNames);
 
 				GLastCompileTime = FPlatformTime::Seconds();
 				GWorkerDiagnostics.BatchProcessTime = GLastCompileTime - BatchProcessStartTime;
 			}
 
+			GWorkerDiagnostics.ErrorCode = FSCWErrorCode::IsSet() ? FSCWErrorCode::Get() : FSCWErrorCode::Success;
+
 			// Prepare for output
 			if (DisableFileWrite)
 			{
 				// write to in-memory bytestream instead for debugging purposes
 				TArray<uint8> MemBlock;
 				FMemoryWriter MemWriter(MemBlock);
-				WriteToOutputArchive(MemWriter, SingleJobs, PipelineJobs, PipelineJobNames);
+				FShaderCompileWorkerUtil::WriteToOutputArchive(MemWriter, GetLocalHostName(), GWorkerDiagnostics, GNumProcessedJobs,
+					SingleJobs, PipelineJobs, PipelineJobNames);
 			}
 			else
 			{
@@ -330,7 +220,8 @@ public:
 				{
 					TUniquePtr<FArchive> OutputFilePtr = TUniquePtr<FArchive>(CreateOutputArchive());
 					check(OutputFilePtr);
-					WriteToOutputArchive(*OutputFilePtr, SingleJobs, PipelineJobs, PipelineJobNames);
+					FShaderCompileWorkerUtil::WriteToOutputArchive(*OutputFilePtr, GetLocalHostName(), GWorkerDiagnostics, GNumProcessedJobs,
+						SingleJobs, PipelineJobs, PipelineJobNames);
 				}
 
 				// Change the output file name to requested one, if we're using a temporary output file
@@ -340,6 +231,9 @@ public:
 				}
 			}
 
+			// Reset error code as it can receive a new value in subsequent jobs
+			FSCWErrorCode::Reset();
+
 			if (IsUsingXGE())
 			{
 				// To signal compilation completion, create a zero length file in the working directory.
@@ -464,7 +358,6 @@ private:
 
 	FString InputFilePath;
 	FString OutputFilePath;
-	TMap<FName, uint32> FormatVersionMap;
 	FString TempFilePath;
 
 	/** Opens an input file, trying multiple times if necessary. */
@@ -488,421 +381,6 @@ private:
 		return InputFile;
 	}
 
-	void VerifyFormatVersions(TMap<FName, uint32>& ReceivedFormatVersionMap)
-	{
-		for (auto Pair : ReceivedFormatVersionMap)
-		{
-			auto* Found = FormatVersionMap.Find(Pair.Key);
-			if (Found)
-			{
-				if (Pair.Value != *Found)
-				{
-					ExitWithoutCrash(FSCWErrorCode::BadShaderFormatVersion, FString::Printf(TEXT("Mismatched shader version for format %s: Found version %u but expected %u; did you forget to build ShaderCompilerWorker?"), *Pair.Key.ToString(), *Found, Pair.Value));
-				}
-			}
-		}
-	}
-
-	void ProcessInputFromArchive(FArchive* InputFilePtr, TArray<FShaderCompileJob>& OutSingleJobs, TArray<FShaderPipelineCompileJob>& OutPipelineJobs, TArray<FString>& OutPipelineNames)
-	{
-		TRACE_CPUPROFILER_EVENT_SCOPE(ProcessInputFromArchive);
-
-		int32 InputVersion;
-		*InputFilePtr << InputVersion;
-		if (ShaderCompileWorkerInputVersion != InputVersion)
-		{
-			ExitWithoutCrash(FSCWErrorCode::BadInputVersion, FString::Printf(TEXT("Exiting due to ShaderCompilerWorker expecting input version %d, got %d instead! Did you forget to build ShaderCompilerWorker?"), ShaderCompileWorkerInputVersion, InputVersion));
-		}
-
-		FString CompressionFormatString;
-		*InputFilePtr << CompressionFormatString;
-		FName CompressionFormat(*CompressionFormatString);
-
-		bool bWasCompressed = (CompressionFormat != NAME_None);
-
-		TArray<uint8> UncompressedData;
-		if (bWasCompressed)
-		{
-			int32 UncompressedDataSize = 0;
-			*InputFilePtr << UncompressedDataSize;
-
-			if (UncompressedDataSize == 0)
-			{
-				ExitWithoutCrash(FSCWErrorCode::BadInputFile, TEXT("Exiting due to bad input file to ShaderCompilerWorker (uncompressed size is 0)! Did you forget to build ShaderCompilerWorker?"));
-				// unreachable
-				return;
-			}
-
-			UncompressedData.SetNumUninitialized(UncompressedDataSize);
-			TArray<uint8> CompressedData;
-			*InputFilePtr << CompressedData;
-			if (!FCompression::UncompressMemory(CompressionFormat, UncompressedData.GetData(), UncompressedDataSize, CompressedData.GetData(), CompressedData.Num()))
-			{
-				ExitWithoutCrash(FSCWErrorCode::BadInputFile, FString::Printf(TEXT("Exiting due to bad input file to ShaderCompilerWorker (cannot uncompress with the format %s)! Did you forget to build ShaderCompilerWorker?"), *CompressionFormatString));
-				// unreachable
-				return;
-			}
-		}
-		FMemoryReader InputMemory(UncompressedData);
-		FArchive& InputFile = bWasCompressed ? InputMemory : *InputFilePtr;
-
-		TMap<FName, uint32> ReceivedFormatVersionMap;
-		InputFile << ReceivedFormatVersionMap;
-
-		bool bIsDebugCompile = !DebugSourceFiles.IsEmpty();
-
-		// when running a debug compile we are more permissive w.r.t. format versions.
-		// (instead of requiring all versions match, we just check the version for the debug compilation
-		// that we will execute)
-		if (!bIsDebugCompile)
-		{
-			VerifyFormatVersions(ReceivedFormatVersionMap);
-		}
-		
-		// Initialize shader hash cache before reading any includes.
-		InitializeShaderHashCache();
-
-		// Array of string used as const TCHAR* during compilation process.
-		TArray<TUniquePtr<FString>> AllocatedStrings;
-		auto DeserializeConstTCHAR = [&AllocatedStrings](FArchive& Archive)
-		{
-			FString Name;
-			Archive << Name;
-
-			const TCHAR* CharName = nullptr;
-			if (Name.Len() != 0)
-			{
-				if (AllocatedStrings.GetSlack() == 0)
-				{
-					AllocatedStrings.Reserve(AllocatedStrings.Num() + 1024);
-				}
-
-				AllocatedStrings.Add(MakeUnique<FString>(Name));
-				CharName = **AllocatedStrings.Last();
-			}
-			return CharName;
-		};
-
-		// Array of string used as const ANSICHAR* during compilation process.
-		TArray<TUniquePtr<TArray<ANSICHAR>>> AllocatedAnsiStrings;
-		auto DeserializeConstANSICHAR = [&AllocatedAnsiStrings](FArchive& Archive)
-		{
-			FString Name;
-			Archive << Name;
-
-			const ANSICHAR* CharName = nullptr;
-			if (Name.Len() != 0)
-			{
-				if (AllocatedAnsiStrings.GetSlack() == 0)
-				{
-					AllocatedAnsiStrings.Reserve(AllocatedAnsiStrings.Num() + 1024);
-				}
-
-				TArray<ANSICHAR> AnsiString;
-				AnsiString.SetNumZeroed(Name.Len() + 1);
-				ANSICHAR* Dest = &AnsiString[0];
-				FCStringAnsi::Strncpy(Dest, TCHAR_TO_ANSI(*Name), Name.Len() + 1);
-
-				AllocatedAnsiStrings.Add(MakeUnique<TArray<ANSICHAR>>(AnsiString));
-				CharName = &(*AllocatedAnsiStrings.Last())[0];
-			}
-			return CharName;
-		};
-
-		// Shared environments
-		TArray<FShaderCompilerEnvironment> SharedEnvironments;
-		{
-			int32 NumSharedEnvironments = 0;
-			InputFile << NumSharedEnvironments;
-			SharedEnvironments.Empty(NumSharedEnvironments);
-			SharedEnvironments.AddDefaulted(NumSharedEnvironments);
-
-			for (int32 EnvironmentIndex = 0; EnvironmentIndex < NumSharedEnvironments; EnvironmentIndex++)
-			{
-				SharedEnvironments[EnvironmentIndex].SerializeCompilationDependencies(InputFile);
-			}
-		}
-
-		// All the shader parameter structures
-		// Note: this is a bit more complicated, purposefully to avoid switch const TCHAR* to FString in runtime FShaderParametersMetadata.
-		TArray<TUniquePtr<FShaderParametersMetadata>> ParameterStructures;
-		{
-			int32 NumParameterStructures = 0;
-			InputFile << NumParameterStructures;
-			ParameterStructures.Reserve(NumParameterStructures);
-
-			for (int32 StructIndex = 0; StructIndex < NumParameterStructures; StructIndex++)
-			{
-				const TCHAR* LayoutName;
-				const TCHAR* StructTypeName;
-				const TCHAR* ShaderVariableName;
-				FShaderParametersMetadata::EUseCase UseCase;
-				const ANSICHAR* StructFileName;
-				int32 StructFileLine;
-				uint32 Size;
-				int32 MemberCount;
-
-				LayoutName = DeserializeConstTCHAR(InputFile);
-				StructTypeName = DeserializeConstTCHAR(InputFile);
-				ShaderVariableName = DeserializeConstTCHAR(InputFile);
-				InputFile << UseCase;
-				StructFileName = DeserializeConstANSICHAR(InputFile);
-				InputFile << StructFileLine;
-				InputFile << Size;
-				InputFile << MemberCount;
-
-				TArray<FShaderParametersMetadata::FMember> Members;
-				Members.Reserve(MemberCount);
-
-				for (int32 MemberIndex = 0; MemberIndex < MemberCount; MemberIndex++)
-				{
-					const TCHAR* Name;
-					const TCHAR* ShaderType;
-					int32 FileLine;
-					uint32 Offset;
-					uint8 BaseType;
-					uint8 PrecisionModifier;
-					uint32 NumRows;
-					uint32 NumColumns;
-					uint32 NumElements;
-					int32 StructMetadataIndex;
-
-					static_assert(sizeof(BaseType) == sizeof(EUniformBufferBaseType), "Cast failure.");
-					static_assert(sizeof(PrecisionModifier) == sizeof(EShaderPrecisionModifier::Type), "Cast failure.");
-
-					Name = DeserializeConstTCHAR(InputFile);
-					ShaderType = DeserializeConstTCHAR(InputFile);
-					InputFile << FileLine;
-					InputFile << Offset;
-					InputFile << BaseType;
-					InputFile << PrecisionModifier;
-					InputFile << NumRows;
-					InputFile << NumColumns;
-					InputFile << NumElements;
-					InputFile << StructMetadataIndex;
-
-					if (ShaderType == nullptr)
-					{
-						ShaderType = TEXT("");
-					}
-
-					const FShaderParametersMetadata* StructMetadata = nullptr;
-					if (StructMetadataIndex != INDEX_NONE)
-					{
-						StructMetadata = ParameterStructures[StructMetadataIndex].Get();
-					}
-
-					FShaderParametersMetadata::FMember Member(
-						Name,
-						ShaderType,
-						FileLine,
-						Offset,
-						EUniformBufferBaseType(BaseType),
-						EShaderPrecisionModifier::Type(PrecisionModifier),
-						NumRows,
-						NumColumns,
-						NumElements,
-						StructMetadata);
-					Members.Add(Member);
-				}
-
-				ParameterStructures.Add(MakeUnique<FShaderParametersMetadata>(
-					UseCase,
-					EUniformBufferBindingFlags::Shader,
-					/* InLayoutName = */ LayoutName,
-					/* InStructTypeName = */ StructTypeName,
-					/* InShaderVariableName = */ ShaderVariableName,
-					/* InStaticSlotName = */ nullptr,
-					StructFileName,
-					StructFileLine,
-					Size,
-					Members,
-					/* bCompleteInitialization = */ true));
-			}
-		}
-
-		GNumProcessedJobs = 0;
-
-		auto SetupDebugCompile = [this, &ReceivedFormatVersionMap](FShaderCompileJob& Job, const TCHAR* SourcePath, const TCHAR* SecondarySourcePath = nullptr)
-			{
-				Job.Input.DebugInfoFlags |= EShaderDebugInfoFlags::CompileFromDebugUSF;
-
-				checkf(ReceivedFormatVersionMap[Job.Input.ShaderFormat] == FormatVersionMap[Job.Input.ShaderFormat],
-						TEXT("Format mismatch for shader format %s; ensure SDK/tools for this platform match the build where this debug compile input was generated.\n")
-						TEXT("This check is skippable, with the caveat that this debug compilation may fail for unexpected reasons and/or results may be different."),
-						*Job.Input.ShaderFormat.ToString());
-
-				FString DebugSource;
-				bool bSuccess = FFileHelper::LoadFileToString(DebugSource, SourcePath);
-
-				if (!DumpDebugInfoPathOverride.IsEmpty())
-				{
-					Job.Input.DumpDebugInfoPath = DumpDebugInfoPathOverride;
-				}
-
-				if (bSuccess)
-				{
-					// strip comments from source when loading from a debug source file. some backends don't handle the comments that the debug dump inserts properly.
-					FAnsiString Stripped;
-					// Since we are directly assigning the stripped result to an FShaderSource object which automatically appends SIMD patting as part of the Set function,
-					// we don't need to also add padding during the convert-and-strip operation.
-					ShaderConvertAndStripComments(DebugSource, Stripped.GetCharArray(), EConvertAndStripFlags::NoSimdPadding);
-					Job.PreprocessOutput.EditSource() = MoveTemp(Stripped);
-				}
-
-				if (SecondarySourcePath != nullptr)
-				{
-					FString SecondaryDebugSource;
-					bSuccess = FFileHelper::LoadFileToString(SecondaryDebugSource, SecondarySourcePath);
-					if (bSuccess)
-					{
-						// strip comments from source when loading from a debug source file. some backends don't handle the comments that the debug dump inserts properly.
-						FAnsiString SecondaryStripped;
-						// Since we are directly assigning the stripped result to an FShaderSource object which automatically appends SIMD patting as part of the Set function,
-						// we don't need to also add padding during the convert-and-strip operation.
-						ShaderConvertAndStripComments(SecondaryDebugSource, SecondaryStripped.GetCharArray(), EConvertAndStripFlags::NoSimdPadding);
-						Job.SecondaryPreprocessOutput->EditSource() = MoveTemp(SecondaryStripped);
-					}
-				}
-			};
-
-		uint32 VirtualSourceIndex = 0;
-		TStringBuilder<512> VirtualSourcePath;
-		auto LoadVirtualSourceFile = [&VirtualSourceIndex, &VirtualSourcePath, this](FShaderPreprocessOutput* PreprocessOutput)
-			{
-				if (PreprocessOutput)
-				{
-					FShaderSource& Source = PreprocessOutput->EditSource();
-					// When running via UBA we use the virtual file mechanism to pass source to avoid duplicating it in memory between 
-					// the input file and job object. from SCW's point of view it just needs to check if the compressed data in the source
-					// object is empty or not to detect if this is the current execution path.
-					checkf(Source.IsCompressed(), TEXT("Shader source sent to SCW was not compressed; this is unexpected"));
-					if (Source.GetCompressed().IsEmpty())
-					{
-						VirtualSourcePath.Reset();
-						VirtualSourcePath << InputFilePath << VirtualSourceIndex++;
-						FFileHelper::LoadFileToArray(Source.EditCompressed(), VirtualSourcePath.ToString(), FILEREAD_NoFail);
-					}
-				}
-			};
-
-		// Individual jobs
-		{
-			int32 SingleJobHeader = ShaderCompileWorkerSingleJobHeader;
-			InputFile << SingleJobHeader;
-			if (ShaderCompileWorkerSingleJobHeader != SingleJobHeader)
-			{
-				ExitWithoutCrash(FSCWErrorCode::BadSingleJobHeader, FString::Printf(TEXT("Exiting due to ShaderCompilerWorker expecting job header %d, got %d instead! Did you forget to build ShaderCompilerWorker?"), ShaderCompileWorkerSingleJobHeader, SingleJobHeader));
-			}
-
-			int32 NumSingleJobs = 0;
-			InputFile << NumSingleJobs;
-
-			FlushShaderFileCache();
-
-			OutSingleJobs.Reserve(NumSingleJobs);
-		
-			for (int32 BatchIndex = 0; BatchIndex < NumSingleJobs; BatchIndex++)
-			{
-				FShaderCompileJob& Job = OutSingleJobs.AddDefaulted_GetRef();
-				// Deserialize the job's inputs.
-				Job.SerializeWorkerInput(InputFile);
-				if (bIsDebugCompile)
-				{
-					bool bHasSecondary = Job.SecondaryPreprocessOutput.IsValid();
-					// this should only be used in debug compile mode; in this case we expect either a single job or a single pipeline job
-					if (bHasSecondary)
-					{
-						// jobs with secondary compiles should have two source files
-						check(DebugSourceFiles.Num() == 2 && NumSingleJobs == 1);
-						SetupDebugCompile(Job, *DebugSourceFiles[0], *DebugSourceFiles[1]);
-					}
-					else
-					{
-						// normal jobs should just have a single source file
-						check(DebugSourceFiles.Num() == 1 && NumSingleJobs == 1);
-						SetupDebugCompile(Job, *DebugSourceFiles[0]);
-					}
-				}
-				else
-				{
-					LoadVirtualSourceFile(&Job.PreprocessOutput);
-					LoadVirtualSourceFile(Job.SecondaryPreprocessOutput.Get());
-				}
-
-				Job.Input.DeserializeSharedInputs(InputFile, SharedEnvironments, ParameterStructures);
-
-				// Process the job.
-				CompileShader(GetShaderFormats(), Job, &GNumProcessedJobs);
-				if (bIsDebugCompile)
-				{
-					// if we're in debug compile mode and we have a single job, skip the pipeline job section below so
-					// we can check that in either debug compile case we have the correct number of debug source files.
-					return;
-				}
-			}
-		}
-
-		// Shader pipeline jobs
-		{
-			int32 PipelineJobHeader = ShaderCompileWorkerPipelineJobHeader;
-			InputFile << PipelineJobHeader;
-			if (ShaderCompileWorkerPipelineJobHeader != PipelineJobHeader)
-			{
-				ExitWithoutCrash(FSCWErrorCode::BadPipelineJobHeader, FString::Printf(TEXT("Exiting due to ShaderCompilerWorker expecting pipeline job header %d, got %d instead! Did you forget to build ShaderCompilerWorker?"), ShaderCompileWorkerSingleJobHeader, PipelineJobHeader));
-			}
-
-			int32 NumPipelines = 0;
-			InputFile << NumPipelines;
-
-			OutPipelineNames.Reserve(NumPipelines);
-			OutPipelineJobs.Reserve(NumPipelines);
-
-			// ensure that we are only compiling a single pipeline if debug source files were specified
-			check(!bIsDebugCompile || NumPipelines == 1);
-
-			for (int32 Index = 0; Index < NumPipelines; ++Index)
-			{
-				FString& PipelineName = OutPipelineNames.AddDefaulted_GetRef();
-				InputFile << PipelineName;
-
-				int32 NumStages = 0;
-				InputFile << NumStages;
-				FShaderPipelineCompileJob& PipelineJob = OutPipelineJobs.Emplace_GetRef(NumStages);
-
-				int32 DebugSourceFileIndex = 0, DebugSourceFileCount = DebugSourceFiles.Num();
-				for (int32 StageIndex = 0; StageIndex < NumStages; ++StageIndex)
-				{
-					// Deserialize the job's inputs.
-					FShaderCompileJob* Job = PipelineJob.StageJobs[StageIndex]->GetSingleShaderJob();
-					Job->SerializeWorkerInput(InputFile);
-					if (bIsDebugCompile)
-					{
-						const TCHAR* DebugSourceFile = DebugSourceFileIndex < DebugSourceFileCount ? *DebugSourceFiles[DebugSourceFileIndex++] : nullptr;
-						const TCHAR* DebugSecondarySourceFile = (Job->SecondaryPreprocessOutput.IsValid() && (DebugSourceFileIndex < DebugSourceFileCount)) ? *DebugSourceFiles[DebugSourceFileIndex++] : nullptr;
-						SetupDebugCompile(*Job, DebugSourceFile, DebugSecondarySourceFile);
-					}
-					else
-					{
-						LoadVirtualSourceFile(&Job->PreprocessOutput);
-						LoadVirtualSourceFile(Job->SecondaryPreprocessOutput.Get());
-					}
-					Job->Input.DeserializeSharedInputs(InputFile, SharedEnvironments, ParameterStructures);
-
-					// SCW doesn't run DDPI, GShaderHasCache Initialize is run  at start with no knowledge of the CustomPlatforms
-					// CustomPlatforms are known when we parse the WorkerInput so we populate the Directory here
-					if (IsCustomPlatform((EShaderPlatform)Job->Input.Target.Platform))
-					{
-						const EShaderPlatform ShaderPlatform = ShaderFormatNameToShaderPlatform(Job->Input.ShaderFormat);
-						UpdateIncludeDirectoryForPreviewPlatform((EShaderPlatform)Job->Input.Target.Platform, ShaderPlatform);
-					}
-				}
-
-				CompileShaderPipeline(GetShaderFormats(), &PipelineJob, &GNumProcessedJobs);
-			}
-		}
-	}
-
 	bool IsUsingTempOutput()
 	{
 		return !TempFilePath.IsEmpty();
@@ -971,60 +449,6 @@ private:
 		return OutputFilePtr;
 	}
 
-	void WriteToOutputArchive(FArchive& OutputFile, TArray<FShaderCompileJob>& SingleJobs, TArray<FShaderPipelineCompileJob>& PipelineJobs, TArray<FString>& PipelineNames)
-	{
-		TRACE_CPUPROFILER_EVENT_SCOPE(WriteToOutputArchive);
-
-		int64 FileSizePosition = 0;
-		if (FSCWErrorCode::IsSet())
-		{
-			GWorkerDiagnostics.ErrorCode = FSCWErrorCode::Get();
-			FileSizePosition = WriteOutputFileHeader(OutputFile, 0, nullptr, FSCWErrorCode::GetInfo().Len(), *FSCWErrorCode::GetInfo());
-		}
-		else
-		{
-			GWorkerDiagnostics.ErrorCode = FSCWErrorCode::Success;
-			FileSizePosition = WriteOutputFileHeader(OutputFile, 0, nullptr, 0, nullptr);
-		}
-
-		{
-			int32 SingleJobHeader = ShaderCompileWorkerSingleJobHeader;
-			OutputFile << SingleJobHeader;
-
-			int32 NumBatches = SingleJobs.Num();
-			OutputFile << NumBatches;
-
-			for (int32 JobIndex = 0; JobIndex < SingleJobs.Num(); JobIndex++)
-			{
-				SingleJobs[JobIndex].SerializeWorkerOutput(OutputFile);
-				UpdateFileSize(OutputFile, FileSizePosition);
-			}
-		}
-
-		{
-			int32 PipelineJobHeader = ShaderCompileWorkerPipelineJobHeader;
-			OutputFile << PipelineJobHeader;
-			int32 NumBatches = PipelineJobs.Num();
-			OutputFile << NumBatches;
-
-			for (int32 JobIndex = 0; JobIndex < PipelineJobs.Num(); JobIndex++)
-			{
-				auto& PipelineJob = PipelineJobs[JobIndex];
-				OutputFile << PipelineNames[JobIndex];
-				bool bSucceeded = (bool)PipelineJob.bSucceeded;
-				OutputFile << bSucceeded;
-				int32 NumStageJobs = PipelineJob.StageJobs.Num();
-				OutputFile << NumStageJobs;
-
-				for (int32 Index = 0; Index < NumStageJobs; ++Index)
-				{
-					PipelineJob.StageJobs[Index]->SerializeWorkerOutput(OutputFile);
-					UpdateFileSize(OutputFile, FileSizePosition);
-				}
-			}
-		}
-	}
-
 	/** Called in the idle loop, checks for conditions under which the helper should exit */
 	void CheckExitConditions()
 	{
@@ -1206,23 +630,9 @@ static int32 GuardedMain(int32 argc, TCHAR* argv[], FString& CrashOutputFile)
 	// Explicitly load ShaderPreprocessor module so it will run its initialization step
 	FModuleManager::LoadModuleChecked<IModuleInterface>(TEXT("ShaderPreprocessor"));
 
-	// We just enumerate the shader formats here for debugging.
-	const TArray<const class IShaderFormat*>& ShaderFormats = GetShaderFormats();
-	check(ShaderFormats.Num());
-	TMap<FName, uint32> FormatVersionMap;
-	for (int32 Index = 0; Index < ShaderFormats.Num(); Index++)
+	if (FShaderCompileWorkerUtil::GetShaderFormats().IsEmpty())
 	{
-		TRACE_CPUPROFILER_EVENT_SCOPE(ShaderFormat);
-
-		TArray<FName> OutFormats;
-		ShaderFormats[Index]->GetSupportedFormats(OutFormats);
-		check(OutFormats.Num());
-		for (int32 InnerIndex = 0; InnerIndex < OutFormats.Num(); InnerIndex++)
-		{
-			UE_LOG(LogShaders, Display, TEXT("Available Shader Format %s"), *OutFormats[InnerIndex].ToString());
-			uint32 Version = ShaderFormats[Index]->GetVersion(OutFormats[InnerIndex]);
-			FormatVersionMap.Add(OutFormats[InnerIndex], Version);
-		}
+		ExitWithoutCrash(FSCWErrorCode::NoTargetShaderFormatsFound, TEXT("No target shader formats found!"));
 	}
 
 	GLastCompileTime = FPlatformTime::Seconds();
@@ -1234,7 +644,7 @@ static int32 GuardedMain(int32 argc, TCHAR* argv[], FString& CrashOutputFile)
 
 	TRACE_CPUPROFILER_EVENT_SCOPE(FWorkLoop);
 
-	FWorkLoop WorkLoop(argv[2], argv[1], argv[4], argv[5], FormatVersionMap);
+	FWorkLoop WorkLoop(argv[2], argv[1], argv[4], argv[5]);
 	WorkLoop.Loop(CrashOutputFile);
 	
 
@@ -1301,13 +711,13 @@ static int32 GuardedMainWrapper(int32 ArgC, TCHAR* ArgV[], FString& CrashOutputF
 					}
 				}
 
-				int64 FileSizePosition = WriteOutputFileHeader(*OutputFile, ExceptionCallStack.Len(), *ExceptionCallStack, ExceptionMsg.Len(), *ExceptionMsg);
+				int64 FileSizePosition = FShaderCompileWorkerUtil::WriteOutputFileHeader(*OutputFile, GetLocalHostName(), GWorkerDiagnostics, GNumProcessedJobs, ExceptionCallStack.Len(), *ExceptionCallStack, ExceptionMsg.Len(), *ExceptionMsg);
 
 				int32 NumBatches = 0;
 				*OutputFile << NumBatches;
 				*OutputFile << NumBatches;
 
-				UpdateFileSize(*OutputFile, FileSizePosition);
+				FShaderCompileWorkerUtil::UpdateFileSize(*OutputFile, FileSizePosition);
 			}
 
 			if (IsUsingXGE())
diff --git a/Engine/Source/Runtime/RenderCore/Internal/ShaderCompileWorkerUtil.h b/Engine/Source/Runtime/RenderCore/Internal/ShaderCompileWorkerUtil.h
--- a/Engine/Source/Runtime/RenderCore/Internal/ShaderCompileWorkerUtil.h
+++ b/Engine/Source/Runtime/RenderCore/Internal/ShaderCompileWorkerUtil.h
@@ -6,6 +6,7 @@
 #include "Containers/Array.h"
 
 class FArchive;
+template <typename FuncType> class TFunctionRef;
 
 struct FSCWErrorCode
 {
@@ -24,6 +25,7 @@ struct FSCWErrorCode
 		CantCompileForSpecificFormat,
 		CrashInsidePlatformCompiler,
 		BadInputFile,
+		MissingVirtualInputFile,
 		OutOfMemory,
 	};
 
@@ -68,6 +70,39 @@ public:
 	RENDERCORE_API static bool WriteTasks(const TArray<FShaderCommonCompileJobPtr>& QueuedJobs, FArchive& TransferFile, EWriteTasksFlags Flags = EWriteTasksFlags::None, TArray<FMemoryView>* OutCompressedSource = nullptr);
 	RENDERCORE_API static FSCWErrorCode::ECode ReadTasks(const TArray<FShaderCommonCompileJobPtr>& QueuedJobs, FArchive& OutputFile, FShaderCompileWorkerDiagnostics* OutWorkerDiagnostics, EReadTasksFlags Flags = EReadTasksFlags::None);
 
+	[[nodiscard]] RENDERCORE_API static const TArray<const IShaderFormat*>& GetShaderFormats();
+
+	RENDERCORE_API static void ProcessInputFromArchive(
+		FArchive* InputFilePtr,
+		const FString& DumpDebugInfoPathOverride,
+		TConstArrayView<FString> DebugSourceFiles,
+		TFunctionRef<TArray<uint8> (int32)> LoadVirtualSourceFile,
+		int32& OutNumProcessedJobs,
+		TArray<FShaderCompileJob>& OutSingleJobs,
+		TArray<FShaderPipelineCompileJob>& OutPipelineJobs,
+		TArray<FString>& OutPipelineNames);
+
+	RENDERCORE_API static void WriteToOutputArchive(
+		FArchive& OutputFile,
+		const FString& HostName,
+		FShaderCompileWorkerDiagnostics& WorkerDiagnostics,
+		int32 NumProcessedJobs,
+		TArray<FShaderCompileJob>& SingleJobs,
+		TArray<FShaderPipelineCompileJob>& PipelineJobs,
+		TArray<FString>& PipelineNames);
+
+	RENDERCORE_API static int64 WriteOutputFileHeader(
+		FArchive& OutputFile,
+		const FString& HostName,
+		FShaderCompileWorkerDiagnostics& WorkerDiagnostics,
+		int32 NumProcessedJobs,
+		int32 CallstackLength,
+		const TCHAR* Callstack,
+		int32 ExceptionInfoLength,
+		const TCHAR* ExceptionInfo);
+
+	RENDERCORE_API static void UpdateFileSize(FArchive& OutputFile, int64 FileSizePosition);
+
 private:
 	static bool HandleWorkerCrash(
 		const TArray<FShaderCommonCompileJobPtr>& QueuedJobs,
diff --git a/Engine/Source/Runtime/RenderCore/Private/ShaderCompileWorkerUtil.cpp b/Engine/Source/Runtime/RenderCore/Private/ShaderCompileWorkerUtil.cpp
--- a/Engine/Source/Runtime/RenderCore/Private/ShaderCompileWorkerUtil.cpp
+++ b/Engine/Source/Runtime/RenderCore/Private/ShaderCompileWorkerUtil.cpp
@@ -2,14 +2,23 @@
 
 #include "ShaderCompileWorkerUtil.h"
 
+#include "Containers/AnsiString.h"
+#include "Containers/Map.h"
+#include "Interfaces/IShaderFormat.h"
+#include "Interfaces/IShaderFormatModule.h"
+#include "Interfaces/ITargetPlatformManagerModule.h"
+#include "Logging/StructuredLog.h"
 #include "Misc/Compression.h"
 #include "Misc/FileHelper.h"
 #include "Misc/MessageDialog.h"
 #include "Misc/OutputDeviceRedirector.h"
+#include "Misc/StringBuilder.h"
+#include "Modules/ModuleManager.h"
+#include "ProfilingDebugging/CpuProfilerTrace.h"
+#include "RHIShaderFormatDefinitions.inl"
+#include "Serialization/MemoryReader.h"
 #include "Serialization/MemoryWriter.h"
 #include "ShaderDiagnostics.h"
-#include "Interfaces/ITargetPlatformManagerModule.h"
-#include "Interfaces/IShaderFormat.h"
 
 static TAutoConsoleVariable<bool> CVarShadersPropagateLocalWorkerOOMs(
 	TEXT("r.Shaders.PropagateLocalWorkerOOMs"),
@@ -196,6 +205,11 @@ namespace ShaderCompileWorkerError
 		ModalErrorOrLog(TEXT("ShaderCompileWorker failed"), FString::Printf(TEXT("Bad-input-file exception:\n%s"), Data));
 	}
 
+	void HandleMissingVirtualInputFile(const TCHAR* Data)
+	{
+		ModalErrorOrLog(TEXT("ShaderCompileWorker failed"), FString::Printf(TEXT("Missing-virtual-input-file exception:\n%s"), Data));
+	}
+
 	bool HandleOutOfMemory(const TCHAR* ExceptionInfo, const TCHAR* Hostname, const FPlatformMemoryStats& MemoryStats, const TArray<FShaderCommonCompileJobPtr>& QueuedJobs, bool bWillRetry)
 	{
 		constexpr int64 Gibibyte = 1024 * 1024 * 1024;
@@ -363,6 +377,9 @@ bool FShaderCompileWorkerUtil::HandleWorkerCrash(
 	case FSCWErrorCode::BadInputFile:
 		ShaderCompileWorkerError::HandleBadInputFile(ExceptionInfo.GetData());
 		break;
+	case FSCWErrorCode::MissingVirtualInputFile:
+		ShaderCompileWorkerError::HandleMissingVirtualInputFile(ExceptionInfo.GetData());
+		break;
 	case FSCWErrorCode::OutOfMemory:
 		return ShaderCompileWorkerError::HandleOutOfMemory(ExceptionInfo.GetData(), Hostname.GetData(), MemoryStats, QueuedJobs, bWillRetry);
 	case FSCWErrorCode::Success:
@@ -886,3 +903,589 @@ FSCWErrorCode::ECode FShaderCompileWorkerUtil::ReadTasks(const TArray<FShaderCom
 	return FSCWErrorCode::Success;
 }
 
+const TArray<const IShaderFormat*>& FShaderCompileWorkerUtil::GetShaderFormats()
+{
+	static const TArray<const IShaderFormat*> Formats = []
+	{
+		TArray<const IShaderFormat*> Output;
+		TArray<FName> Modules;
+		FModuleManager::Get().FindModules(SHADERFORMAT_MODULE_WILDCARD, Modules);
+		for (FName Module : Modules)
+		{
+			if (IShaderFormat* Format = FModuleManager::LoadModuleChecked<IShaderFormatModule>(Module).GetShaderFormat())
+			{
+				Output.Add(Format);
+			}
+		}
+		return Output;
+	}();
+	return Formats;
+}
+
+// Like GetFormatVersionMap above but does not rely on TargetPlatformManager.
+[[nodiscard]] static const TMap<FName, uint32>& GetShaderFormatVersionMap()
+{
+	static const TMap<FName, uint32> Versions = []
+	{
+		TMap<FName, uint32> Output;
+		for (const IShaderFormat* ShaderFormat : FShaderCompileWorkerUtil::GetShaderFormats())
+		{
+			TRACE_CPUPROFILER_EVENT_SCOPE(ShaderFormat);
+			TArray<FName> Formats;
+			ShaderFormat->GetSupportedFormats(Formats);
+			for (FName Format : Formats)
+			{
+				UE_LOGFMT(LogShaders, Display, "Available Shader Format {Format}", Format);
+				uint32 Version = ShaderFormat->GetVersion(Format);
+				Output.Add(Format, Version);
+			}
+		}
+		return Output;
+	}();
+	return Versions;
+}
+
+[[nodiscard]] static bool VerifyFormatVersions(const TMap<FName, uint32>& FormatVersionMap, TMap<FName, uint32>& ReceivedFormatVersionMap)
+{
+	for (const TTuple<FName, uint32>& Pair : ReceivedFormatVersionMap)
+	{
+		if (const uint32* Expected = FormatVersionMap.Find(Pair.Key))
+		{
+			if (Pair.Value != *Expected)
+			{
+				FSCWErrorCode::Report(FSCWErrorCode::BadShaderFormatVersion, FString::Printf(TEXT("Mismatched shader version for format %s: Found version %u but expected %u; did you forget to build ShaderCompilerWorker?"), *Pair.Key.ToString(), Pair.Value, *Expected));
+				return false;
+			}
+		}
+	}
+	return true;
+}
+
+void FShaderCompileWorkerUtil::ProcessInputFromArchive(
+	FArchive* InputFilePtr,
+	const FString& DumpDebugInfoPathOverride,
+	TConstArrayView<FString> DebugSourceFiles,
+	TFunctionRef<TArray<uint8> (int32)> LoadVirtualSourceFile,
+	int32& OutNumProcessedJobs,
+	TArray<FShaderCompileJob>& OutSingleJobs,
+	TArray<FShaderPipelineCompileJob>& OutPipelineJobs,
+	TArray<FString>& OutPipelineNames)
+{
+	TRACE_CPUPROFILER_EVENT_SCOPE(ProcessInputFromArchive);
+
+	int32 InputVersion;
+	*InputFilePtr << InputVersion;
+	if (ShaderCompileWorkerInputVersion != InputVersion)
+	{
+		FSCWErrorCode::Report(FSCWErrorCode::BadInputVersion, FString::Printf(TEXT("Exiting due to ShaderCompilerWorker expecting input version %d, got %d instead! Did you forget to build ShaderCompilerWorker?"), ShaderCompileWorkerInputVersion, InputVersion));
+		return;
+	}
+
+	FString CompressionFormatString;
+	*InputFilePtr << CompressionFormatString;
+	FName CompressionFormat(*CompressionFormatString);
+
+	const bool bWasCompressed = !CompressionFormat.IsNone();
+
+	TArray<uint8> UncompressedData;
+	if (bWasCompressed)
+	{
+		int32 UncompressedDataSize = 0;
+		*InputFilePtr << UncompressedDataSize;
+
+		if (UncompressedDataSize == 0)
+		{
+			FSCWErrorCode::Report(FSCWErrorCode::BadInputFile, TEXT("Exiting due to bad input file to ShaderCompilerWorker (uncompressed size is 0)! Did you forget to build ShaderCompilerWorker?"));
+			return;
+		}
+
+		UncompressedData.SetNumUninitialized(UncompressedDataSize);
+		TArray<uint8> CompressedData;
+		*InputFilePtr << CompressedData;
+		if (!FCompression::UncompressMemory(CompressionFormat, UncompressedData.GetData(), UncompressedDataSize, CompressedData.GetData(), CompressedData.Num()))
+		{
+			FSCWErrorCode::Report(FSCWErrorCode::BadInputFile, FString::Printf(TEXT("Exiting due to bad input file to ShaderCompilerWorker (cannot uncompress with the format %s)! Did you forget to build ShaderCompilerWorker?"), *CompressionFormatString));
+			return;
+		}
+	}
+	FMemoryReader InputMemory(UncompressedData);
+	FArchive& InputFile = bWasCompressed ? InputMemory : *InputFilePtr;
+
+	TMap<FName, uint32> ReceivedFormatVersionMap;
+	InputFile << ReceivedFormatVersionMap;
+
+	const TMap<FName, uint32>& FormatVersionMap = GetShaderFormatVersionMap();
+
+	const bool bIsDebugCompile = !DebugSourceFiles.IsEmpty();
+
+	// When running a debug compile we are more permissive w.r.t. format versions.
+	// Instead of requiring all versions match, we check only the version for the debug compilation that we will execute.
+	if (!bIsDebugCompile && !VerifyFormatVersions(FormatVersionMap, ReceivedFormatVersionMap))
+	{
+		return;
+	}
+
+	// Initialize shader hash cache before reading any includes.
+	InitializeShaderHashCache();
+
+	// Array of string used as const TCHAR* during compilation process.
+	TArray<TUniquePtr<FString>> AllocatedStrings;
+	auto DeserializeConstTCHAR = [&AllocatedStrings](FArchive& Archive)
+	{
+		FString Name;
+		Archive << Name;
+
+		const TCHAR* CharName = nullptr;
+		if (Name.Len() != 0)
+		{
+			if (AllocatedStrings.GetSlack() == 0)
+			{
+				AllocatedStrings.Reserve(AllocatedStrings.Num() + 1024);
+			}
+			CharName = **AllocatedStrings.Add_GetRef(MakeUnique<FString>(MoveTemp(Name)));
+		}
+		return CharName;
+	};
+
+	// Array of string used as const ANSICHAR* during compilation process.
+	TArray<TUniquePtr<FAnsiString>> AllocatedAnsiStrings;
+	auto DeserializeConstANSICHAR = [&AllocatedAnsiStrings](FArchive& Archive)
+	{
+		FString Name;
+		Archive << Name;
+
+		const ANSICHAR* CharName = nullptr;
+		if (Name.Len() != 0)
+		{
+			if (AllocatedAnsiStrings.GetSlack() == 0)
+			{
+				AllocatedAnsiStrings.Reserve(AllocatedAnsiStrings.Num() + 1024);
+			}
+			CharName = **AllocatedAnsiStrings.Add_GetRef(MakeUnique<FAnsiString>(MoveTemp(Name)));
+		}
+		return CharName;
+	};
+
+	// Shared environments
+	TArray<FShaderCompilerEnvironment> SharedEnvironments;
+	{
+		int32 NumSharedEnvironments = 0;
+		InputFile << NumSharedEnvironments;
+		SharedEnvironments.Empty(NumSharedEnvironments);
+		SharedEnvironments.AddDefaulted(NumSharedEnvironments);
+
+		for (int32 EnvironmentIndex = 0; EnvironmentIndex < NumSharedEnvironments; EnvironmentIndex++)
+		{
+			SharedEnvironments[EnvironmentIndex].SerializeCompilationDependencies(InputFile);
+		}
+	}
+
+	// All the shader parameter structures
+	// Note: this is a bit more complicated, purposefully to avoid switch const TCHAR* to FString in runtime FShaderParametersMetadata.
+	TArray<TUniquePtr<FShaderParametersMetadata>> ParameterStructures;
+	{
+		int32 NumParameterStructures = 0;
+		InputFile << NumParameterStructures;
+		ParameterStructures.Reserve(NumParameterStructures);
+
+		for (int32 StructIndex = 0; StructIndex < NumParameterStructures; StructIndex++)
+		{
+			const TCHAR* LayoutName;
+			const TCHAR* StructTypeName;
+			const TCHAR* ShaderVariableName;
+			FShaderParametersMetadata::EUseCase UseCase;
+			const ANSICHAR* StructFileName;
+			int32 StructFileLine;
+			uint32 Size;
+			int32 MemberCount;
+
+			LayoutName = DeserializeConstTCHAR(InputFile);
+			StructTypeName = DeserializeConstTCHAR(InputFile);
+			ShaderVariableName = DeserializeConstTCHAR(InputFile);
+			InputFile << UseCase;
+			StructFileName = DeserializeConstANSICHAR(InputFile);
+			InputFile << StructFileLine;
+			InputFile << Size;
+			InputFile << MemberCount;
+
+			TArray<FShaderParametersMetadata::FMember> Members;
+			Members.Reserve(MemberCount);
+
+			for (int32 MemberIndex = 0; MemberIndex < MemberCount; MemberIndex++)
+			{
+				const TCHAR* Name;
+				const TCHAR* ShaderType;
+				int32 FileLine;
+				uint32 Offset;
+				uint8 BaseType;
+				uint8 PrecisionModifier;
+				uint32 NumRows;
+				uint32 NumColumns;
+				uint32 NumElements;
+				int32 StructMetadataIndex;
+
+				static_assert(sizeof(BaseType) == sizeof(EUniformBufferBaseType), "Cast failure.");
+				static_assert(sizeof(PrecisionModifier) == sizeof(EShaderPrecisionModifier::Type), "Cast failure.");
+
+				Name = DeserializeConstTCHAR(InputFile);
+				ShaderType = DeserializeConstTCHAR(InputFile);
+				InputFile << FileLine;
+				InputFile << Offset;
+				InputFile << BaseType;
+				InputFile << PrecisionModifier;
+				InputFile << NumRows;
+				InputFile << NumColumns;
+				InputFile << NumElements;
+				InputFile << StructMetadataIndex;
+
+				if (ShaderType == nullptr)
+				{
+					ShaderType = TEXT("");
+				}
+
+				const FShaderParametersMetadata* StructMetadata = nullptr;
+				if (StructMetadataIndex != INDEX_NONE)
+				{
+					StructMetadata = ParameterStructures[StructMetadataIndex].Get();
+				}
+
+				FShaderParametersMetadata::FMember Member(
+					Name,
+					ShaderType,
+					FileLine,
+					Offset,
+					EUniformBufferBaseType(BaseType),
+					EShaderPrecisionModifier::Type(PrecisionModifier),
+					NumRows,
+					NumColumns,
+					NumElements,
+					StructMetadata);
+				Members.Add(Member);
+			}
+
+			ParameterStructures.Add(MakeUnique<FShaderParametersMetadata>(
+				UseCase,
+				EUniformBufferBindingFlags::Shader,
+				/* InLayoutName = */ LayoutName,
+				/* InStructTypeName = */ StructTypeName,
+				/* InShaderVariableName = */ ShaderVariableName,
+				/* InStaticSlotName = */ nullptr,
+				StructFileName,
+				StructFileLine,
+				Size,
+				Members,
+				/* bForceCompleteInitialization = */ true));
+		}
+	}
+
+	OutNumProcessedJobs = 0;
+
+	auto SetupDebugCompile = [&FormatVersionMap, &ReceivedFormatVersionMap, &DumpDebugInfoPathOverride](FShaderCompileJob& Job, const TCHAR* SourcePath, const TCHAR* SecondarySourcePath = nullptr)
+	{
+		Job.Input.DebugInfoFlags |= EShaderDebugInfoFlags::CompileFromDebugUSF;
+
+		checkf(FormatVersionMap.FindRef(Job.Input.ShaderFormat, 0) == ReceivedFormatVersionMap.FindRef(Job.Input.ShaderFormat, 0),
+			TEXT("Format mismatch for shader format %s; ensure SDK/tools for this platform match the build where this debug compile input was generated.\n")
+			TEXT("This check is skippable, with the caveat that this debug compilation may fail for unexpected reasons and/or results may be different."),
+			*Job.Input.ShaderFormat.ToString());
+
+		if (!DumpDebugInfoPathOverride.IsEmpty())
+		{
+			Job.Input.DumpDebugInfoPath = DumpDebugInfoPathOverride;
+		}
+
+		auto LoadDebugSource = [](const TCHAR* Path, FShaderPreprocessOutput& Output)
+		{
+			if (FString DebugSource; FFileHelper::LoadFileToString(DebugSource, Path))
+			{
+				// Strip comments from source when loading from a debug source file. Some backends don't handle the comments that the debug dump inserts properly.
+				FAnsiString Stripped;
+				// Since we are directly assigning the stripped result to an FShaderSource object which automatically appends SIMD patting as part of the Set function,
+				// we don't need to also add padding during the convert-and-strip operation.
+				ShaderConvertAndStripComments(DebugSource, Stripped.GetCharArray(), EConvertAndStripFlags::NoSimdPadding);
+				Output.EditSource() = MoveTemp(Stripped);
+			}
+		};
+
+		LoadDebugSource(SourcePath, Job.PreprocessOutput);
+
+		if (SecondarySourcePath && Job.SecondaryPreprocessOutput)
+		{
+			LoadDebugSource(SecondarySourcePath, *Job.SecondaryPreprocessOutput);
+		}
+	};
+
+	auto TryLoadVirtualSourceFile = [LoadVirtualSourceFile, VirtualSourceIndex = 0](FShaderPreprocessOutput* PreprocessOutput) mutable
+	{
+		if (PreprocessOutput)
+		{
+			// The virtual file mechanism is used by some build controllers to avoid duplicating source data between the
+			// input file and job object. Use of the virtual file mechanism is indicated by the absence of compressed source data.
+			if (FShaderSource& Source = PreprocessOutput->EditSource(); Source.GetCompressed().IsEmpty())
+			{
+				TArray<uint8> Compressed = LoadVirtualSourceFile(VirtualSourceIndex);
+				if (Compressed.IsEmpty())
+				{
+					FSCWErrorCode::Report(FSCWErrorCode::MissingVirtualInputFile, FString::Printf(TEXT("Failed to load virtual source file at index %d."), VirtualSourceIndex));
+					return false;
+				}
+				Source.EditCompressed() = MoveTemp(Compressed);
+				++VirtualSourceIndex;
+			}
+		}
+		return true;
+	};
+
+	// Individual jobs
+	{
+		int32 SingleJobHeader = ShaderCompileWorkerSingleJobHeader;
+		InputFile << SingleJobHeader;
+		if (ShaderCompileWorkerSingleJobHeader != SingleJobHeader)
+		{
+			FSCWErrorCode::Report(FSCWErrorCode::BadSingleJobHeader, FString::Printf(TEXT("Exiting due to ShaderCompilerWorker expecting job header %d, got %d instead! Did you forget to build ShaderCompilerWorker?"), ShaderCompileWorkerSingleJobHeader, SingleJobHeader));
+			return;
+		}
+
+		int32 NumSingleJobs = 0;
+		InputFile << NumSingleJobs;
+
+		FlushShaderFileCache();
+
+		OutSingleJobs.Reserve(NumSingleJobs);
+		
+		for (int32 BatchIndex = 0; BatchIndex < NumSingleJobs; BatchIndex++)
+		{
+			FShaderCompileJob& Job = OutSingleJobs.AddDefaulted_GetRef();
+			// Deserialize the job's inputs.
+			Job.SerializeWorkerInput(InputFile);
+			if (bIsDebugCompile)
+			{
+				bool bHasSecondary = Job.SecondaryPreprocessOutput.IsValid();
+				// this should only be used in debug compile mode; in this case we expect either a single job or a single pipeline job
+				if (bHasSecondary)
+				{
+					// jobs with secondary compiles should have two source files
+					check(DebugSourceFiles.Num() == 2 && NumSingleJobs == 1);
+					SetupDebugCompile(Job, *DebugSourceFiles[0], *DebugSourceFiles[1]);
+				}
+				else
+				{
+					// normal jobs should just have a single source file
+					check(DebugSourceFiles.Num() == 1 && NumSingleJobs == 1);
+					SetupDebugCompile(Job, *DebugSourceFiles[0]);
+				}
+			}
+			else
+			{
+				if (!TryLoadVirtualSourceFile(&Job.PreprocessOutput) ||
+					!TryLoadVirtualSourceFile(Job.SecondaryPreprocessOutput.Get()))
+				{
+					return;
+				}
+			}
+
+			Job.Input.DeserializeSharedInputs(InputFile, SharedEnvironments, ParameterStructures);
+
+			// Process the job.
+			CompileShader(GetShaderFormats(), Job, &OutNumProcessedJobs);
+			if (bIsDebugCompile)
+			{
+				// if we're in debug compile mode and we have a single job, skip the pipeline job section below so
+				// we can check that in either debug compile case we have the correct number of debug source files.
+				return;
+			}
+		}
+	}
+
+	// Shader pipeline jobs
+	{
+		int32 PipelineJobHeader = ShaderCompileWorkerPipelineJobHeader;
+		InputFile << PipelineJobHeader;
+		if (ShaderCompileWorkerPipelineJobHeader != PipelineJobHeader)
+		{
+			FSCWErrorCode::Report(FSCWErrorCode::BadPipelineJobHeader, FString::Printf(TEXT("Exiting due to ShaderCompilerWorker expecting pipeline job header %d, got %d instead! Did you forget to build ShaderCompilerWorker?"), ShaderCompileWorkerSingleJobHeader, PipelineJobHeader));
+			return;
+		}
+
+		int32 NumPipelines = 0;
+		InputFile << NumPipelines;
+
+		OutPipelineNames.Reserve(NumPipelines);
+		OutPipelineJobs.Reserve(NumPipelines);
+
+		// ensure that we are only compiling a single pipeline if debug source files were specified
+		check(!bIsDebugCompile || NumPipelines == 1);
+
+		for (int32 Index = 0; Index < NumPipelines; ++Index)
+		{
+			FString& PipelineName = OutPipelineNames.AddDefaulted_GetRef();
+			InputFile << PipelineName;
+
+			int32 NumStages = 0;
+			InputFile << NumStages;
+			FShaderPipelineCompileJob& PipelineJob = OutPipelineJobs.Emplace_GetRef(NumStages);
+
+			int32 DebugSourceFileIndex = 0, DebugSourceFileCount = DebugSourceFiles.Num();
+			for (int32 StageIndex = 0; StageIndex < NumStages; ++StageIndex)
+			{
+				// Deserialize the job's inputs.
+				FShaderCompileJob* Job = PipelineJob.StageJobs[StageIndex]->GetSingleShaderJob();
+				Job->SerializeWorkerInput(InputFile);
+				if (bIsDebugCompile)
+				{
+					const TCHAR* DebugSourceFile = DebugSourceFileIndex < DebugSourceFileCount ? *DebugSourceFiles[DebugSourceFileIndex++] : nullptr;
+					const TCHAR* DebugSecondarySourceFile = (Job->SecondaryPreprocessOutput.IsValid() && (DebugSourceFileIndex < DebugSourceFileCount)) ? *DebugSourceFiles[DebugSourceFileIndex++] : nullptr;
+					SetupDebugCompile(*Job, DebugSourceFile, DebugSecondarySourceFile);
+				}
+				else
+				{
+					if (!TryLoadVirtualSourceFile(&Job->PreprocessOutput) ||
+						!TryLoadVirtualSourceFile(Job->SecondaryPreprocessOutput.Get()))
+					{
+						return;
+					}
+				}
+				Job->Input.DeserializeSharedInputs(InputFile, SharedEnvironments, ParameterStructures);
+
+				// SCW doesn't run DDPI, GShaderHasCache Initialize is run at start with no knowledge of the CustomPlatforms
+				// CustomPlatforms are known when we parse the WorkerInput so we populate the Directory here
+				if (IsCustomPlatform((EShaderPlatform)Job->Input.Target.Platform))
+				{
+					const EShaderPlatform ShaderPlatform = ShaderFormatNameToShaderPlatform(Job->Input.ShaderFormat);
+					UpdateIncludeDirectoryForPreviewPlatform((EShaderPlatform)Job->Input.Target.Platform, ShaderPlatform);
+				}
+			}
+
+			CompileShaderPipeline(GetShaderFormats(), &PipelineJob, &OutNumProcessedJobs);
+		}
+	}
+}
+
+void FShaderCompileWorkerUtil::UpdateFileSize(FArchive& OutputFile, int64 FileSizePosition)
+{
+	int64 Current = OutputFile.Tell();
+	OutputFile.Seek(FileSizePosition);
+	OutputFile << Current;
+	OutputFile.Seek(Current);
+}
+
+int64 FShaderCompileWorkerUtil::WriteOutputFileHeader(
+	FArchive& OutputFile,
+	const FString& HostName,
+	FShaderCompileWorkerDiagnostics& WorkerDiagnostics,
+	int32 NumProcessedJobs,
+	int32 CallstackLength,
+	const TCHAR* Callstack,
+	int32 ExceptionInfoLength,
+	const TCHAR* ExceptionInfo)
+{
+	TRACE_CPUPROFILER_EVENT_SCOPE(WriteOutputFileHeader);
+
+	int64 FileSizePosition = 0;
+	int32 OutputVersion = ShaderCompileWorkerOutputVersion;
+	OutputFile << OutputVersion;
+
+	int64 FileSize = 0;
+	// Get the position of the Size value to be patched in as the shader progresses
+	FileSizePosition = OutputFile.Tell();
+	OutputFile << FileSize;
+
+	OutputFile << WorkerDiagnostics;
+
+	OutputFile << NumProcessedJobs;
+
+	// Note: Can't use FStrings here as SEH can't be used with destructors
+	OutputFile << CallstackLength;
+
+	OutputFile << ExceptionInfoLength;
+
+	int32 HostNameLength = HostName.Len();
+	OutputFile << HostNameLength;
+
+	if (WorkerDiagnostics.ErrorCode != FSCWErrorCode::Success)
+	{
+		if (CallstackLength > 0)
+		{
+			OutputFile.Serialize((void*)Callstack, CallstackLength * sizeof(TCHAR));
+		}
+
+		if (ExceptionInfoLength > 0)
+		{
+			OutputFile.Serialize((void*)ExceptionInfo, ExceptionInfoLength * sizeof(TCHAR));
+		}
+
+		if (HostNameLength > 0)
+		{
+			OutputFile.Serialize((void*)*HostName, HostNameLength * sizeof(TCHAR));
+		}
+
+		// Store available and used physical memory of host machine on OOM error
+		if (WorkerDiagnostics.ErrorCode == FSCWErrorCode::OutOfMemory)
+		{
+			FPlatformMemoryStats MemoryStats = FPlatformMemory::GetStats();
+
+			OutputFile
+				<< MemoryStats.AvailablePhysical
+				<< MemoryStats.AvailableVirtual
+				<< MemoryStats.UsedPhysical
+				<< MemoryStats.PeakUsedPhysical
+				<< MemoryStats.UsedVirtual
+				<< MemoryStats.PeakUsedVirtual
+				;
+		}
+	}
+
+	FShaderCompileWorkerUtil::UpdateFileSize(OutputFile, FileSizePosition);
+	return FileSizePosition;
+}
+
+void FShaderCompileWorkerUtil::WriteToOutputArchive(
+	FArchive& OutputFile,
+	const FString& HostName,
+	FShaderCompileWorkerDiagnostics& WorkerDiagnostics,
+	int32 NumProcessedJobs,
+	TArray<FShaderCompileJob>& SingleJobs,
+	TArray<FShaderPipelineCompileJob>& PipelineJobs,
+	TArray<FString>& PipelineNames)
+{
+	TRACE_CPUPROFILER_EVENT_SCOPE(WriteToOutputArchive);
+
+	const int64 FileSizePosition = FShaderCompileWorkerUtil::WriteOutputFileHeader(OutputFile, HostName, WorkerDiagnostics, NumProcessedJobs, 0, nullptr, FSCWErrorCode::GetInfo().Len(), *FSCWErrorCode::GetInfo());
+
+	{
+		int32 SingleJobHeader = ShaderCompileWorkerSingleJobHeader;
+		OutputFile << SingleJobHeader;
+
+		int32 NumBatches = SingleJobs.Num();
+		OutputFile << NumBatches;
+
+		for (int32 JobIndex = 0; JobIndex < SingleJobs.Num(); JobIndex++)
+		{
+			SingleJobs[JobIndex].SerializeWorkerOutput(OutputFile);
+			FShaderCompileWorkerUtil::UpdateFileSize(OutputFile, FileSizePosition);
+		}
+	}
+
+	{
+		int32 PipelineJobHeader = ShaderCompileWorkerPipelineJobHeader;
+		OutputFile << PipelineJobHeader;
+		int32 NumBatches = PipelineJobs.Num();
+		OutputFile << NumBatches;
+
+		for (int32 JobIndex = 0; JobIndex < PipelineJobs.Num(); JobIndex++)
+		{
+			auto& PipelineJob = PipelineJobs[JobIndex];
+			OutputFile << PipelineNames[JobIndex];
+			bool bSucceeded = (bool)PipelineJob.bSucceeded;
+			OutputFile << bSucceeded;
+			int32 NumStageJobs = PipelineJob.StageJobs.Num();
+			OutputFile << NumStageJobs;
+
+			for (int32 Index = 0; Index < NumStageJobs; ++Index)
+			{
+				PipelineJob.StageJobs[Index]->SerializeWorkerOutput(OutputFile);
+				FShaderCompileWorkerUtil::UpdateFileSize(OutputFile, FileSizePosition);
+			}
+		}
+	}
+}
