[shaders] disable UE_ENABLE_TRACKED_IO for ShaderCompileWorker; this is unnecessary overhead since it only affects the new console window which SCW does not use.

#Synced-CL 52132623

diff --git a/Engine/Source/Programs/ShaderCompileWorker/ShaderCompileWorker.Build.cs b/Engine/Source/Programs/ShaderCompileWorker/ShaderCompileWorker.Build.cs
--- a/Engine/Source/Programs/ShaderCompileWorker/ShaderCompileWorker.Build.cs
+++ b/Engine/Source/Programs/ShaderCompileWorker/ShaderCompileWorker.Build.cs
@@ -24,6 +24,9 @@ public class ShaderCompileWorker : ModuleRules
 				"Sockets",
 			});
 
+		// SCW has no console window, so tracked IO/activity visualization is unused overhead
+		PublicDefinitions.Add("UE_ENABLE_TRACKED_IO=0");
+
 		PrivateIncludePathModuleNames.AddRange(
 			new string[] {
 				"Launch",
